Fix merge issue in CL 18841574

diff --git a/Templates/TP_ThirdPerson/Config/DefaultEngine.ini b/Templates/TP_ThirdPerson/Config/DefaultEngine.ini
--- a/Templates/TP_ThirdPerson/Config/DefaultEngine.ini
+++ b/Templates/TP_ThirdPerson/Config/DefaultEngine.ini
@@ -4,10 +4,4 @@ EditorStartupMap=/Game/ThirdPerson/Maps/ThirdPersonMap.ThirdPersonMap
 GlobalDefaultGameMode="/Script/TP_ThirdPerson.TP_ThirdPersonGameMode"
 
 [/Script/IOSRuntimeSettings.IOSRuntimeSettings]
->>>> ORIGINAL //UE5/Release-5.0/Templates/TP_ThirdPerson/Config/DefaultEngine.ini#1
-MinimumiOSVersion=IOS_13
-==== THEIRS //UE5/Release-5.0/Templates/TP_ThirdPerson/Config/DefaultEngine.ini#2
-MinimumiOSVersion=IOS_13
-
-==== YOURS //Axel-Release5-iMac/Templates/TP_ThirdPerson/Config/DefaultEngine.ini
 MinimumiOSVersion=IOS_14
