[Templates] Switching 3rd Person C++ Template to enhanced input

#content only

diff --git a/Templates/TP_ThirdPerson/Config/DefaultInput.ini b/Templates/TP_ThirdPerson/Config/DefaultInput.ini
--- a/Templates/TP_ThirdPerson/Config/DefaultInput.ini
+++ b/Templates/TP_ThirdPerson/Config/DefaultInput.ini
@@ -1,22 +1,75 @@
 
+
 [/Script/Engine.InputSettings]
-+AxisConfig=(AxisKeyName="Gamepad_LeftX",AxisProperties=(DeadZone=0.250000,Sensitivity=1.000000,Exponent=1.000000,bInvert=False))
-+AxisConfig=(AxisKeyName="Gamepad_LeftY",AxisProperties=(DeadZone=0.250000,Sensitivity=1.000000,Exponent=1.000000,bInvert=False))
-+AxisConfig=(AxisKeyName="Gamepad_RightX",AxisProperties=(DeadZone=0.250000,Sensitivity=1.000000,Exponent=1.000000,bInvert=False))
-+AxisConfig=(AxisKeyName="Gamepad_RightY",AxisProperties=(DeadZone=0.250000,Sensitivity=1.000000,Exponent=1.000000,bInvert=False))
-+AxisConfig=(AxisKeyName="MouseX",AxisProperties=(DeadZone=0.000000,Sensitivity=0.070000,Exponent=1.000000,bInvert=False))
+-AxisConfig=(AxisKeyName="Gamepad_LeftX",AxisProperties=(DeadZone=0.25,Exponent=1.f,Sensitivity=1.f))
+-AxisConfig=(AxisKeyName="Gamepad_LeftY",AxisProperties=(DeadZone=0.25,Exponent=1.f,Sensitivity=1.f))
+-AxisConfig=(AxisKeyName="Gamepad_RightX",AxisProperties=(DeadZone=0.25,Exponent=1.f,Sensitivity=1.f))
+-AxisConfig=(AxisKeyName="Gamepad_RightY",AxisProperties=(DeadZone=0.25,Exponent=1.f,Sensitivity=1.f))
+-AxisConfig=(AxisKeyName="MouseX",AxisProperties=(DeadZone=0.f,Exponent=1.f,Sensitivity=0.07f))
+-AxisConfig=(AxisKeyName="MouseY",AxisProperties=(DeadZone=0.f,Exponent=1.f,Sensitivity=0.07f))
+-AxisConfig=(AxisKeyName="Mouse2D",AxisProperties=(DeadZone=0.f,Exponent=1.f,Sensitivity=0.07f))
 +AxisConfig=(AxisKeyName="MouseY",AxisProperties=(DeadZone=0.000000,Sensitivity=0.070000,Exponent=1.000000,bInvert=False))
 +AxisConfig=(AxisKeyName="MouseWheelAxis",AxisProperties=(DeadZone=0.000000,Sensitivity=1.000000,Exponent=1.000000,bInvert=False))
 +AxisConfig=(AxisKeyName="Gamepad_LeftTriggerAxis",AxisProperties=(DeadZone=0.000000,Sensitivity=1.000000,Exponent=1.000000,bInvert=False))
 +AxisConfig=(AxisKeyName="Gamepad_RightTriggerAxis",AxisProperties=(DeadZone=0.000000,Sensitivity=1.000000,Exponent=1.000000,bInvert=False))
 +AxisConfig=(AxisKeyName="Gamepad_Special_Left_X",AxisProperties=(DeadZone=0.000000,Sensitivity=1.000000,Exponent=1.000000,bInvert=False))
 +AxisConfig=(AxisKeyName="Gamepad_Special_Left_Y",AxisProperties=(DeadZone=0.000000,Sensitivity=1.000000,Exponent=1.000000,bInvert=False))
++AxisConfig=(AxisKeyName="Mouse2D",AxisProperties=(DeadZone=0.000000,Sensitivity=0.070000,Exponent=1.000000,bInvert=False))
++AxisConfig=(AxisKeyName="Gamepad_LeftX",AxisProperties=(DeadZone=0.250000,Sensitivity=1.000000,Exponent=1.000000,bInvert=False))
++AxisConfig=(AxisKeyName="Gamepad_LeftY",AxisProperties=(DeadZone=0.250000,Sensitivity=1.000000,Exponent=1.000000,bInvert=False))
++AxisConfig=(AxisKeyName="Gamepad_RightX",AxisProperties=(DeadZone=0.250000,Sensitivity=1.000000,Exponent=1.000000,bInvert=False))
++AxisConfig=(AxisKeyName="Gamepad_RightY",AxisProperties=(DeadZone=0.250000,Sensitivity=1.000000,Exponent=1.000000,bInvert=False))
++AxisConfig=(AxisKeyName="MouseX",AxisProperties=(DeadZone=0.000000,Sensitivity=0.070000,Exponent=1.000000,bInvert=False))
++AxisConfig=(AxisKeyName="Vive_Left_Trigger_Axis",AxisProperties=(DeadZone=0.000000,Sensitivity=1.000000,Exponent=1.000000,bInvert=False))
++AxisConfig=(AxisKeyName="Vive_Left_Trackpad_X",AxisProperties=(DeadZone=0.000000,Sensitivity=1.000000,Exponent=1.000000,bInvert=False))
++AxisConfig=(AxisKeyName="Vive_Left_Trackpad_Y",AxisProperties=(DeadZone=0.000000,Sensitivity=1.000000,Exponent=1.000000,bInvert=False))
++AxisConfig=(AxisKeyName="Vive_Right_Trigger_Axis",AxisProperties=(DeadZone=0.000000,Sensitivity=1.000000,Exponent=1.000000,bInvert=False))
++AxisConfig=(AxisKeyName="Vive_Right_Trackpad_X",AxisProperties=(DeadZone=0.000000,Sensitivity=1.000000,Exponent=1.000000,bInvert=False))
++AxisConfig=(AxisKeyName="Vive_Right_Trackpad_Y",AxisProperties=(DeadZone=0.000000,Sensitivity=1.000000,Exponent=1.000000,bInvert=False))
++AxisConfig=(AxisKeyName="MixedReality_Left_Trigger_Axis",AxisProperties=(DeadZone=0.000000,Sensitivity=1.000000,Exponent=1.000000,bInvert=False))
++AxisConfig=(AxisKeyName="MixedReality_Left_Thumbstick_X",AxisProperties=(DeadZone=0.000000,Sensitivity=1.000000,Exponent=1.000000,bInvert=False))
++AxisConfig=(AxisKeyName="MixedReality_Left_Thumbstick_Y",AxisProperties=(DeadZone=0.000000,Sensitivity=1.000000,Exponent=1.000000,bInvert=False))
++AxisConfig=(AxisKeyName="MixedReality_Left_Trackpad_X",AxisProperties=(DeadZone=0.000000,Sensitivity=1.000000,Exponent=1.000000,bInvert=False))
++AxisConfig=(AxisKeyName="MixedReality_Left_Trackpad_Y",AxisProperties=(DeadZone=0.000000,Sensitivity=1.000000,Exponent=1.000000,bInvert=False))
++AxisConfig=(AxisKeyName="MixedReality_Right_Trigger_Axis",AxisProperties=(DeadZone=0.000000,Sensitivity=1.000000,Exponent=1.000000,bInvert=False))
++AxisConfig=(AxisKeyName="MixedReality_Right_Thumbstick_X",AxisProperties=(DeadZone=0.000000,Sensitivity=1.000000,Exponent=1.000000,bInvert=False))
++AxisConfig=(AxisKeyName="MixedReality_Right_Thumbstick_Y",AxisProperties=(DeadZone=0.000000,Sensitivity=1.000000,Exponent=1.000000,bInvert=False))
++AxisConfig=(AxisKeyName="MixedReality_Right_Trackpad_X",AxisProperties=(DeadZone=0.000000,Sensitivity=1.000000,Exponent=1.000000,bInvert=False))
++AxisConfig=(AxisKeyName="MixedReality_Right_Trackpad_Y",AxisProperties=(DeadZone=0.000000,Sensitivity=1.000000,Exponent=1.000000,bInvert=False))
++AxisConfig=(AxisKeyName="OculusTouch_Left_Grip_Axis",AxisProperties=(DeadZone=0.000000,Sensitivity=1.000000,Exponent=1.000000,bInvert=False))
++AxisConfig=(AxisKeyName="OculusTouch_Left_Trigger_Axis",AxisProperties=(DeadZone=0.000000,Sensitivity=1.000000,Exponent=1.000000,bInvert=False))
++AxisConfig=(AxisKeyName="OculusTouch_Left_Thumbstick_X",AxisProperties=(DeadZone=0.000000,Sensitivity=1.000000,Exponent=1.000000,bInvert=False))
++AxisConfig=(AxisKeyName="OculusTouch_Left_Thumbstick_Y",AxisProperties=(DeadZone=0.000000,Sensitivity=1.000000,Exponent=1.000000,bInvert=False))
++AxisConfig=(AxisKeyName="OculusTouch_Right_Grip_Axis",AxisProperties=(DeadZone=0.000000,Sensitivity=1.000000,Exponent=1.000000,bInvert=False))
++AxisConfig=(AxisKeyName="OculusTouch_Right_Trigger_Axis",AxisProperties=(DeadZone=0.000000,Sensitivity=1.000000,Exponent=1.000000,bInvert=False))
++AxisConfig=(AxisKeyName="OculusTouch_Right_Thumbstick_X",AxisProperties=(DeadZone=0.000000,Sensitivity=1.000000,Exponent=1.000000,bInvert=False))
++AxisConfig=(AxisKeyName="OculusTouch_Right_Thumbstick_Y",AxisProperties=(DeadZone=0.000000,Sensitivity=1.000000,Exponent=1.000000,bInvert=False))
++AxisConfig=(AxisKeyName="ValveIndex_Left_Grip_Axis",AxisProperties=(DeadZone=0.000000,Sensitivity=1.000000,Exponent=1.000000,bInvert=False))
++AxisConfig=(AxisKeyName="ValveIndex_Left_Grip_Force",AxisProperties=(DeadZone=0.000000,Sensitivity=1.000000,Exponent=1.000000,bInvert=False))
++AxisConfig=(AxisKeyName="ValveIndex_Left_Trigger_Axis",AxisProperties=(DeadZone=0.000000,Sensitivity=1.000000,Exponent=1.000000,bInvert=False))
++AxisConfig=(AxisKeyName="ValveIndex_Left_Thumbstick_X",AxisProperties=(DeadZone=0.000000,Sensitivity=1.000000,Exponent=1.000000,bInvert=False))
++AxisConfig=(AxisKeyName="ValveIndex_Left_Thumbstick_Y",AxisProperties=(DeadZone=0.000000,Sensitivity=1.000000,Exponent=1.000000,bInvert=False))
++AxisConfig=(AxisKeyName="ValveIndex_Left_Trackpad_X",AxisProperties=(DeadZone=0.000000,Sensitivity=1.000000,Exponent=1.000000,bInvert=False))
++AxisConfig=(AxisKeyName="ValveIndex_Left_Trackpad_Y",AxisProperties=(DeadZone=0.000000,Sensitivity=1.000000,Exponent=1.000000,bInvert=False))
++AxisConfig=(AxisKeyName="ValveIndex_Left_Trackpad_Force",AxisProperties=(DeadZone=0.000000,Sensitivity=1.000000,Exponent=1.000000,bInvert=False))
++AxisConfig=(AxisKeyName="ValveIndex_Right_Grip_Axis",AxisProperties=(DeadZone=0.000000,Sensitivity=1.000000,Exponent=1.000000,bInvert=False))
++AxisConfig=(AxisKeyName="ValveIndex_Right_Grip_Force",AxisProperties=(DeadZone=0.000000,Sensitivity=1.000000,Exponent=1.000000,bInvert=False))
++AxisConfig=(AxisKeyName="ValveIndex_Right_Trigger_Axis",AxisProperties=(DeadZone=0.000000,Sensitivity=1.000000,Exponent=1.000000,bInvert=False))
++AxisConfig=(AxisKeyName="ValveIndex_Right_Thumbstick_X",AxisProperties=(DeadZone=0.000000,Sensitivity=1.000000,Exponent=1.000000,bInvert=False))
++AxisConfig=(AxisKeyName="ValveIndex_Right_Thumbstick_Y",AxisProperties=(DeadZone=0.000000,Sensitivity=1.000000,Exponent=1.000000,bInvert=False))
++AxisConfig=(AxisKeyName="ValveIndex_Right_Trackpad_X",AxisProperties=(DeadZone=0.000000,Sensitivity=1.000000,Exponent=1.000000,bInvert=False))
++AxisConfig=(AxisKeyName="ValveIndex_Right_Trackpad_Y",AxisProperties=(DeadZone=0.000000,Sensitivity=1.000000,Exponent=1.000000,bInvert=False))
++AxisConfig=(AxisKeyName="ValveIndex_Right_Trackpad_Force",AxisProperties=(DeadZone=0.000000,Sensitivity=1.000000,Exponent=1.000000,bInvert=False))
 bAltEnterTogglesFullscreen=True
 bF11TogglesFullscreen=True
 bUseMouseForTouch=False
 bEnableMouseSmoothing=True
 bEnableFOVScaling=True
 bCaptureMouseOnLaunch=True
+bEnableLegacyInputScales=True
+bEnableMotionControls=True
+bFilterInputByPlatformUser=False
+bShouldFlushPressedKeysOnViewportFocusLost=True
 bAlwaysShowTouchInterface=False
 bShowConsoleOnFourFingerTap=True
 bEnableGestureRecognizer=False
@@ -25,20 +78,9 @@ DefaultViewportMouseCaptureMode=CapturePermanently_IncludingInitialMouseDown
 DefaultViewportMouseLockMode=LockOnCapture
 FOVScale=0.011110
 DoubleClickTime=0.200000
-+ActionMappings=(ActionName="Jump",bShift=False,bCtrl=False,bAlt=False,bCmd=False,Key=SpaceBar)
-+ActionMappings=(ActionName="Jump",bShift=False,bCtrl=False,bAlt=False,bCmd=False,Key=Gamepad_FaceButton_Bottom)
-+AxisMappings=(AxisName="Move Forward / Backward",Scale=1.000000,Key=W)
-+AxisMappings=(AxisName="Move Forward / Backward",Scale=-1.000000,Key=S)
-+AxisMappings=(AxisName="Move Forward / Backward",Scale=1.000000,Key=Gamepad_LeftY)
-+AxisMappings=(AxisName="Move Right / Left",Scale=-1.000000,Key=A)
-+AxisMappings=(AxisName="Move Right / Left",Scale=1.000000,Key=D)
-+AxisMappings=(AxisName="Move Right / Left",Scale=1.000000,Key=Gamepad_LeftX)
-+AxisMappings=(AxisName="Turn Right / Left Gamepad",Scale=1.000000,Key=Gamepad_RightX)
-+AxisMappings=(AxisName="Turn Right / Left Mouse",Scale=1.000000,Key=MouseX)
-+AxisMappings=(AxisName="Look Up / Down Gamepad",Scale=1.000000,Key=Gamepad_RightY)
-+AxisMappings=(AxisName="Look Up / Down Mouse",Scale=-1.000000,Key=MouseY)
-DefaultTouchInterface=/Engine/MobileResources/HUD/DefaultVirtualJoysticks.DefaultVirtualJoysticks
-+ConsoleKeys=Tilde
 DefaultPlayerInputClass=/Script/EnhancedInput.EnhancedPlayerInput
 DefaultInputComponentClass=/Script/EnhancedInput.EnhancedInputComponent
+DefaultTouchInterface=/Engine/MobileResources/HUD/DefaultVirtualJoysticks.DefaultVirtualJoysticks
+-ConsoleKeys=Tilde
++ConsoleKeys=Tilde
 
diff --git a/Templates/TP_ThirdPerson/Source/TP_ThirdPerson/TP_ThirdPerson.Build.cs b/Templates/TP_ThirdPerson/Source/TP_ThirdPerson/TP_ThirdPerson.Build.cs
--- a/Templates/TP_ThirdPerson/Source/TP_ThirdPerson/TP_ThirdPerson.Build.cs
+++ b/Templates/TP_ThirdPerson/Source/TP_ThirdPerson/TP_ThirdPerson.Build.cs
@@ -8,6 +8,6 @@ public class TP_ThirdPerson : ModuleRules
 	{
 		PCHUsage = PCHUsageMode.UseExplicitOrSharedPCHs;
 
-		PublicDependencyModuleNames.AddRange(new string[] { "Core", "CoreUObject", "Engine", "InputCore", "HeadMountedDisplay" });
+		PublicDependencyModuleNames.AddRange(new string[] { "Core", "CoreUObject", "Engine", "InputCore", "HeadMountedDisplay", "EnhancedInput" });
 	}
 }
diff --git a/Templates/TP_ThirdPerson/Source/TP_ThirdPerson/TP_ThirdPersonCharacter.cpp b/Templates/TP_ThirdPerson/Source/TP_ThirdPerson/TP_ThirdPersonCharacter.cpp
--- a/Templates/TP_ThirdPerson/Source/TP_ThirdPerson/TP_ThirdPersonCharacter.cpp
+++ b/Templates/TP_ThirdPerson/Source/TP_ThirdPerson/TP_ThirdPersonCharacter.cpp
@@ -7,6 +7,9 @@
 #include "GameFramework/CharacterMovementComponent.h"
 #include "GameFramework/Controller.h"
 #include "GameFramework/SpringArmComponent.h"
+#include "EnhancedInputComponent.h"
+#include "EnhancedInputSubsystems.h"
+
 
 //////////////////////////////////////////////////////////////////////////
 // ATP_ThirdPersonCharacter
@@ -15,10 +18,7 @@ ATP_ThirdPersonCharacter::ATP_ThirdPersonCharacter()
 {
 	// Set size for collision capsule
 	GetCapsuleComponent()->InitCapsuleSize(42.f, 96.0f);
-
-	// set our turn rate for input
-	TurnRateGamepad = 50.f;
-
+		
 	// Don't rotate when the controller rotates. Let that just affect the camera.
 	bUseControllerRotationPitch = false;
 	bUseControllerRotationYaw = false;
@@ -51,6 +51,18 @@ ATP_ThirdPersonCharacter::ATP_ThirdPersonCharacter()
 	// are set in the derived blueprint asset named ThirdPersonCharacter (to avoid direct content references in C++)
 }
 
+void ATP_ThirdPersonCharacter::PossessedBy(AController* NewController)
+{
+	//Add a mapping context on Possessed
+	if (APlayerController* PlayerController = Cast<APlayerController>(NewController))
+	{
+		if (UEnhancedInputLocalPlayerSubsystem* Subsystem = ULocalPlayer::GetSubsystem<UEnhancedInputLocalPlayerSubsystem>(PlayerController->GetLocalPlayer()))
+		{
+			Subsystem->AddMappingContext(DefaultMappingContext, 0);
+		}
+	}
+}
+
 //////////////////////////////////////////////////////////////////////////
 // Input
 
@@ -58,72 +70,60 @@ void ATP_ThirdPersonCharacter::SetupPlayerInputComponent(class UInputComponent*
 {
 	// Set up gameplay key bindings
 	check(PlayerInputComponent);
-	PlayerInputComponent->BindAction("Jump", IE_Pressed, this, &ACharacter::Jump);
-	PlayerInputComponent->BindAction("Jump", IE_Released, this, &ACharacter::StopJumping);
-
-	PlayerInputComponent->BindAxis("Move Forward / Backward", this, &ATP_ThirdPersonCharacter::MoveForward);
-	PlayerInputComponent->BindAxis("Move Right / Left", this, &ATP_ThirdPersonCharacter::MoveRight);
-
-	// We have 2 versions of the rotation bindings to handle different kinds of devices differently
-	// "turn" handles devices that provide an absolute delta, such as a mouse.
-	// "turnrate" is for devices that we choose to treat as a rate of change, such as an analog joystick
-	PlayerInputComponent->BindAxis("Turn Right / Left Mouse", this, &APawn::AddControllerYawInput);
-	PlayerInputComponent->BindAxis("Turn Right / Left Gamepad", this, &ATP_ThirdPersonCharacter::TurnAtRate);
-	PlayerInputComponent->BindAxis("Look Up / Down Mouse", this, &APawn::AddControllerPitchInput);
-	PlayerInputComponent->BindAxis("Look Up / Down Gamepad", this, &ATP_ThirdPersonCharacter::LookUpAtRate);
-
-	// handle touch devices
-	PlayerInputComponent->BindTouch(IE_Pressed, this, &ATP_ThirdPersonCharacter::TouchStarted);
-	PlayerInputComponent->BindTouch(IE_Released, this, &ATP_ThirdPersonCharacter::TouchStopped);
-}
+	
+	// Set up action bindings
+	if (UEnhancedInputComponent* EnhancedInputComponent = CastChecked<UEnhancedInputComponent>(PlayerInputComponent)) {
+		
+		//Jumping
+		EnhancedInputComponent->BindAction(JumpAction, ETriggerEvent::Triggered, this, &ACharacter::Jump);
+		EnhancedInputComponent->BindAction(JumpAction, ETriggerEvent::Completed, this, &ACharacter::StopJumping);
 
-void ATP_ThirdPersonCharacter::TouchStarted(ETouchIndex::Type FingerIndex, FVector Location)
-{
-	Jump();
-}
+		//Moving
+		EnhancedInputComponent->BindAction(MoveAction, ETriggerEvent::Triggered, this, &ATP_ThirdPersonCharacter::Move);
 
-void ATP_ThirdPersonCharacter::TouchStopped(ETouchIndex::Type FingerIndex, FVector Location)
-{
-	StopJumping();
-}
+		//Looking
+		EnhancedInputComponent->BindAction(LookAction, ETriggerEvent::Triggered, this, &ATP_ThirdPersonCharacter::Look);
 
-void ATP_ThirdPersonCharacter::TurnAtRate(float Rate)
-{
-	// calculate delta for this frame from the rate information
-	AddControllerYawInput(Rate * TurnRateGamepad * GetWorld()->GetDeltaSeconds());
-}
+	}
 
-void ATP_ThirdPersonCharacter::LookUpAtRate(float Rate)
-{
-	// calculate delta for this frame from the rate information
-	AddControllerPitchInput(Rate * TurnRateGamepad * GetWorld()->GetDeltaSeconds());
 }
 
-void ATP_ThirdPersonCharacter::MoveForward(float Value)
+void ATP_ThirdPersonCharacter::Move(const FInputActionValue& Value)
 {
-	if ((Controller != nullptr) && (Value != 0.0f))
+	// input is a Vector2D
+	FVector2D MovementVector = Value.Get<FVector2D>();
+
+	if (Controller != nullptr)
 	{
 		// find out which way is forward
 		const FRotator Rotation = Controller->GetControlRotation();
 		const FRotator YawRotation(0, Rotation.Yaw, 0);
 
 		// get forward vector
-		const FVector Direction = FRotationMatrix(YawRotation).GetUnitAxis(EAxis::X);
-		AddMovementInput(Direction, Value);
+		const FVector ForwardDirection = FRotationMatrix(YawRotation).GetUnitAxis(EAxis::X);
+	
+		// get right vector 
+		const FVector RightDirection = FRotationMatrix(YawRotation).GetUnitAxis(EAxis::Y);
+
+		// add movement 
+		AddMovementInput(ForwardDirection, MovementVector.Y);
+		AddMovementInput(RightDirection, MovementVector.X);
 	}
 }
 
-void ATP_ThirdPersonCharacter::MoveRight(float Value)
+void ATP_ThirdPersonCharacter::Look(const FInputActionValue& Value)
 {
-	if ( (Controller != nullptr) && (Value != 0.0f) )
+	// input is a Vector2D
+	FVector2D LookAxisVector = Value.Get<FVector2D>();
+
+	if (Controller != nullptr)
 	{
-		// find out which way is right
-		const FRotator Rotation = Controller->GetControlRotation();
-		const FRotator YawRotation(0, Rotation.Yaw, 0);
-	
-		// get right vector 
-		const FVector Direction = FRotationMatrix(YawRotation).GetUnitAxis(EAxis::Y);
-		// add movement in that direction
-		AddMovementInput(Direction, Value);
+		// add yaw and pitch input to controller
+		AddControllerYawInput(LookAxisVector.X);
+		AddControllerPitchInput(LookAxisVector.Y);
 	}
 }
+
+
+
+
diff --git a/Templates/TP_ThirdPerson/Source/TP_ThirdPerson/TP_ThirdPersonCharacter.h b/Templates/TP_ThirdPerson/Source/TP_ThirdPerson/TP_ThirdPersonCharacter.h
--- a/Templates/TP_ThirdPerson/Source/TP_ThirdPerson/TP_ThirdPersonCharacter.h
+++ b/Templates/TP_ThirdPerson/Source/TP_ThirdPerson/TP_ThirdPersonCharacter.h
@@ -4,8 +4,10 @@
 
 #include "CoreMinimal.h"
 #include "GameFramework/Character.h"
+#include "InputActionValue.h"
 #include "TP_ThirdPersonCharacter.generated.h"
 
+
 UCLASS(config=Game)
 class ATP_ThirdPersonCharacter : public ACharacter
 {
@@ -18,38 +20,38 @@ class ATP_ThirdPersonCharacter : public ACharacter
 	/** Follow camera */
 	UPROPERTY(VisibleAnywhere, BlueprintReadOnly, Category = Camera, meta = (AllowPrivateAccess = "true"))
 	class UCameraComponent* FollowCamera;
-public:
-	ATP_ThirdPersonCharacter();
+	
+	/** MappingContext */
+	UPROPERTY(EditAnywhere, BlueprintReadOnly, Category = Input, meta = (AllowPrivateAccess = "true"))
+	class UInputMappingContext* DefaultMappingContext;
 
-	/** Base turn rate, in deg/sec. Other scaling may affect final turn rate. */
-	UPROPERTY(VisibleAnywhere, BlueprintReadOnly, Category=Input)
-	float TurnRateGamepad;
+	/** Jump Input Action */
+	UPROPERTY(EditAnywhere, BlueprintReadOnly, Category = Input, meta = (AllowPrivateAccess = "true"))
+	class UInputAction* JumpAction;
 
-protected:
+	/** Move Input Action */
+	UPROPERTY(EditAnywhere, BlueprintReadOnly, Category = Input, meta = (AllowPrivateAccess = "true"))
+	class UInputAction* MoveAction;
 
-	/** Called for forwards/backward input */
-	void MoveForward(float Value);
+	/** Look Input Action */
+	UPROPERTY(EditAnywhere, BlueprintReadOnly, Category = Input, meta = (AllowPrivateAccess = "true"))
+	class UInputAction* LookAction;
 
-	/** Called for side to side input */
-	void MoveRight(float Value);
+public:
+	ATP_ThirdPersonCharacter();
 
-	/** 
-	 * Called via input to turn at a given rate. 
-	 * @param Rate	This is a normalized rate, i.e. 1.0 means 100% of desired turn rate
-	 */
-	void TurnAtRate(float Rate);
+	/** Called for adding mapping context*/
+	void PossessedBy(AController* NewController) override;
+		
 
-	/**
-	 * Called via input to turn look up/down at a given rate. 
-	 * @param Rate	This is a normalized rate, i.e. 1.0 means 100% of desired turn rate
-	 */
-	void LookUpAtRate(float Rate);
+protected:
 
-	/** Handler for when a touch input begins. */
-	void TouchStarted(ETouchIndex::Type FingerIndex, FVector Location);
+	/** Called for movement input */
+	void Move(const FInputActionValue& Value);
 
-	/** Handler for when a touch input stops. */
-	void TouchStopped(ETouchIndex::Type FingerIndex, FVector Location);
+	/** Called for looking input */
+	void Look(const FInputActionValue& Value);
+			
 
 protected:
 	// APawn interface
diff --git a/Templates/TP_ThirdPersonBP/Config/DefaultEngine.ini b/Templates/TP_ThirdPersonBP/Config/DefaultEngine.ini
--- a/Templates/TP_ThirdPersonBP/Config/DefaultEngine.ini
+++ b/Templates/TP_ThirdPersonBP/Config/DefaultEngine.ini
@@ -19,3 +19,46 @@ r.Lumen.TraceMeshSDFs=0
 r.Shadow.Virtual.Enable=1
 r.Mobile.EnableNoPrecomputedLightingCSMShader=1
 
+[/Script/WindowsTargetPlatform.WindowsTargetSettings]
+DefaultGraphicsRHI=DefaultGraphicsRHI_Default
+-D3D12TargetedShaderFormats=PCD3D_SM5
++D3D12TargetedShaderFormats=PCD3D_SM6
+-D3D11TargetedShaderFormats=PCD3D_SM5
++D3D11TargetedShaderFormats=PCD3D_SM5
+Compiler=Default
+AudioSampleRate=48000
+AudioCallbackBufferFrameSize=1024
+AudioNumBuffersToEnqueue=1
+AudioMaxChannels=0
+AudioNumSourceWorkers=4
+SpatializationPlugin=
+SourceDataOverridePlugin=
+ReverbPlugin=
+OcclusionPlugin=
+CompressionOverrides=(bOverrideCompressionTimes=False,DurationThreshold=5.000000,MaxNumRandomBranches=0,SoundCueQualityIndex=0)
+CacheSizeKB=65536
+MaxChunkSizeOverrideKB=0
+bResampleForDevice=False
+MaxSampleRate=48000.000000
+HighSampleRate=32000.000000
+MedSampleRate=24000.000000
+LowSampleRate=12000.000000
+MinSampleRate=8000.000000
+CompressionQualityModifier=1.000000
+AutoStreamingThreshold=0.000000
+SoundCueCookQualityIndex=-1
+
+[/Script/AndroidFileServerEditor.AndroidFileServerRuntimeSettings]
+bEnablePlugin=True
+bAllowNetworkConnection=True
+SecurityToken=22840D9C40E8F8CFF3280483D90DEE70
+bIncludeInShipping=False
+bAllowExternalStartInShipping=False
+bCompileAFSProject=False
+bUseCompression=False
+bLogFiles=False
+bReportStats=False
+ConnectionType=USBOnly
+bUseManualIPAddress=False
+ManualIPAddress=
+
