Updating TP_ThirdPerson and TP_TopDown templates to default to DX12 and include SM6.

diff --git a/Templates/TP_ThirdPerson/Config/DefaultEngine.ini b/Templates/TP_ThirdPerson/Config/DefaultEngine.ini
--- a/Templates/TP_ThirdPerson/Config/DefaultEngine.ini
+++ b/Templates/TP_ThirdPerson/Config/DefaultEngine.ini
@@ -7,3 +7,46 @@ GlobalDefaultGameMode="/Script/TP_ThirdPerson.TP_ThirdPersonGameMode"
 r.Shadow.Virtual.Enable=1
 r.Mobile.EnableNoPrecomputedLightingCSMShader=1
 
+[/Script/WindowsTargetPlatform.WindowsTargetSettings]
+DefaultGraphicsRHI=DefaultGraphicsRHI_DX12
+-D3D12TargetedShaderFormats=PCD3D_SM5
++D3D12TargetedShaderFormats=PCD3D_SM6
+-D3D11TargetedShaderFormats=PCD3D_SM5
++D3D11TargetedShaderFormats=PCD3D_SM5
+Compiler=Default
+AudioSampleRate=48000
+AudioCallbackBufferFrameSize=1024
+AudioNumBuffersToEnqueue=1
+AudioMaxChannels=0
+AudioNumSourceWorkers=4
+SpatializationPlugin=
+SourceDataOverridePlugin=
+ReverbPlugin=
+OcclusionPlugin=
+CompressionOverrides=(bOverrideCompressionTimes=False,DurationThreshold=5.000000,MaxNumRandomBranches=0,SoundCueQualityIndex=0)
+CacheSizeKB=65536
+MaxChunkSizeOverrideKB=0
+bResampleForDevice=False
+MaxSampleRate=48000.000000
+HighSampleRate=32000.000000
+MedSampleRate=24000.000000
+LowSampleRate=12000.000000
+MinSampleRate=8000.000000
+CompressionQualityModifier=1.000000
+AutoStreamingThreshold=0.000000
+SoundCueCookQualityIndex=-1
+
+[/Script/AndroidFileServerEditor.AndroidFileServerRuntimeSettings]
+bEnablePlugin=True
+bAllowNetworkConnection=True
+SecurityToken=BF6D5A5549DB55B6261E86A891EB414B
+bIncludeInShipping=False
+bAllowExternalStartInShipping=False
+bCompileAFSProject=False
+bUseCompression=False
+bLogFiles=False
+bReportStats=False
+ConnectionType=USBOnly
+bUseManualIPAddress=False
+ManualIPAddress=
+
diff --git a/Templates/TP_ThirdPersonBP/Config/DefaultEngine.ini b/Templates/TP_ThirdPersonBP/Config/DefaultEngine.ini
--- a/Templates/TP_ThirdPersonBP/Config/DefaultEngine.ini
+++ b/Templates/TP_ThirdPersonBP/Config/DefaultEngine.ini
@@ -21,7 +21,7 @@ r.Mobile.EnableNoPrecomputedLightingCSMShader=1
 r.DefaultFeature.AutoExposure.ExtendDefaultLuminanceRange=true
 
 [/Script/WindowsTargetPlatform.WindowsTargetSettings]
-DefaultGraphicsRHI=DefaultGraphicsRHI_Default
+DefaultGraphicsRHI=DefaultGraphicsRHI_DX12
 -D3D12TargetedShaderFormats=PCD3D_SM5
 +D3D12TargetedShaderFormats=PCD3D_SM6
 -D3D11TargetedShaderFormats=PCD3D_SM5
diff --git a/Templates/TP_TopDown/Config/DefaultEngine.ini b/Templates/TP_TopDown/Config/DefaultEngine.ini
--- a/Templates/TP_TopDown/Config/DefaultEngine.ini
+++ b/Templates/TP_TopDown/Config/DefaultEngine.ini
@@ -22,7 +22,7 @@ bUseManualIPAddress=False
 ManualIPAddress=
 
 [/Script/WindowsTargetPlatform.WindowsTargetSettings]
-DefaultGraphicsRHI=DefaultGraphicsRHI_Default
+DefaultGraphicsRHI=DefaultGraphicsRHI_DX12
 -D3D12TargetedShaderFormats=PCD3D_SM5
 +D3D12TargetedShaderFormats=PCD3D_SM6
 -D3D11TargetedShaderFormats=PCD3D_SM5
diff --git a/Templates/TP_TopDownBP/Config/DefaultEngine.ini b/Templates/TP_TopDownBP/Config/DefaultEngine.ini
--- a/Templates/TP_TopDownBP/Config/DefaultEngine.ini
+++ b/Templates/TP_TopDownBP/Config/DefaultEngine.ini
@@ -72,7 +72,7 @@ bUseManualIPAddress=False
 ManualIPAddress=
 
 [/Script/WindowsTargetPlatform.WindowsTargetSettings]
-DefaultGraphicsRHI=DefaultGraphicsRHI_Default
+DefaultGraphicsRHI=DefaultGraphicsRHI_DX12
 -D3D12TargetedShaderFormats=PCD3D_SM5
 +D3D12TargetedShaderFormats=PCD3D_SM6
 -D3D11TargetedShaderFormats=PCD3D_SM5
