[Templates] Updated Physic Asset for Mannequins in shared folders (including some nice profiles now), fixed some warnings, renamed mappings to use IMC_ as prefix, several smaller fixes

#content only

diff --git a/Templates/TP_ThirdPerson/Config/DefaultEngine.ini b/Templates/TP_ThirdPerson/Config/DefaultEngine.ini
--- a/Templates/TP_ThirdPerson/Config/DefaultEngine.ini
+++ b/Templates/TP_ThirdPerson/Config/DefaultEngine.ini
@@ -4,8 +4,13 @@ EditorStartupMap=/Game/ThirdPerson/Maps/ThirdPersonMap.ThirdPersonMap
 GlobalDefaultGameMode="/Script/TP_ThirdPerson.TP_ThirdPersonGameMode"
 
 [/Script/Engine.RendererSettings]
+r.ReflectionMethod=1
+r.GenerateMeshDistanceFields=True
+r.DynamicGlobalIlluminationMethod=1
+r.Lumen.TraceMeshSDFs=0
 r.Shadow.Virtual.Enable=1
 r.Mobile.EnableNoPrecomputedLightingCSMShader=1
+r.DefaultFeature.AutoExposure.ExtendDefaultLuminanceRange=true
 
 [/Script/WindowsTargetPlatform.WindowsTargetSettings]
 DefaultGraphicsRHI=DefaultGraphicsRHI_DX12
