updating zhCN template description for TP_ThirdPerson and TP_ThirdPersionBP

diff --git a/Templates/TP_ThirdPerson/Config/TemplateDefs.ini b/Templates/TP_ThirdPerson/Config/TemplateDefs.ini
--- a/Templates/TP_ThirdPerson/Config/TemplateDefs.ini
+++ b/Templates/TP_ThirdPerson/Config/TemplateDefs.ini
@@ -13,7 +13,7 @@ LocalizedDisplayNames=(Language="ja",Text="サードパーソン")
 LocalizedDescriptions=(Language="ja",Text="サードパーソン パックは、後ろ側と少し上にカメラが配置されているプレイ可能なキャラクターを備えています。タッチデバイスでマウス、キーボード、コントローラー、またはバーチャル ジョイスティックの使用によりキャラクターが移動すると、カメラはキャラクターを追跡し、肩越しの視点から観察します。この視点では主人公が強調されるため、アクションゲームやアドベンチャー ゲームで人気があります。テンプレート マップには、プレイヤーがジャンプしたり押すことができるオブジェクトがいくつか含まれています。")
 
 LocalizedDisplayNames=(Language="zh-Hans",Text="第三人称游戏")
-LocalizedDescriptions=(Language="zh-Hans",Text="第三人称游戏内容包包含一个可操控的角色以及一台位于角色身后斜上方的摄像机。你可以通过鼠标、键盘、手柄或触摸式设备上控制角色移动。摄像机会跟随角色，以越肩视角拍摄场景。这种视角突出主角，常用于动作及冒险类游戏。该模板的地图中还包含一些对象，可以让角色跳上跳下或四处推动。")
+LocalizedDescriptions=(Language="zh-Hans",Text="第三人称游戏内容包包含一个可操控的角色以及一台位于角色身后斜上方的摄像机。你可以通过鼠标、键盘、手柄或触摸式设备上的虚拟摇杆控制角色移动。摄像机会跟随角色，以越肩视角拍摄场景。这种视角突出主角，常用于动作及冒险类游戏。该模板的地图中还包含一些对象，可以让角色跳上跳下或四处推动。")
 
 ; This project uses a thumbnail as the project icon
 bThumbnailAsIcon=true
diff --git a/Templates/TP_ThirdPerson/manifest.json b/Templates/TP_ThirdPerson/manifest.json
--- a/Templates/TP_ThirdPerson/manifest.json
+++ b/Templates/TP_ThirdPerson/manifest.json
@@ -31,7 +31,7 @@
 		},
 		{
 			"Language": "zh-Hans",
-			"Text": "第三人称游戏内容包包含一个可操控的角色以及一台位于角色身后斜上方的摄像机。你可以通过鼠标、键盘、手柄或触摸式设备上控制角色移动。摄像机会跟随角色，以越肩视角拍摄场景。这种视角突出主角，常用于动作及冒险类游戏。该模板的地图中还包含一些对象，可以让角色跳上跳下或四处推动。"
+			"Text": "第三人称游戏内容包包含一个可操控的角色以及一台位于角色身后斜上方的摄像机。你可以通过鼠标、键盘、手柄或触摸式设备上的虚拟摇杆控制角色移动。摄像机会跟随角色，以越肩视角拍摄场景。这种视角突出主角，常用于动作及冒险类游戏。该模板的地图中还包含一些对象，可以让角色跳上跳下或四处推动。"
 		},
 		{
 			"Language": "ja",
diff --git a/Templates/TP_ThirdPersonBP/manifest.json b/Templates/TP_ThirdPersonBP/manifest.json
--- a/Templates/TP_ThirdPersonBP/manifest.json
+++ b/Templates/TP_ThirdPersonBP/manifest.json
@@ -31,7 +31,7 @@
 		},
 		{
 			"Language": "zh-Hans",
-			"Text": "第三人称游戏内容包包含一个可操控的角色以及一台位于角色身后斜上方的摄像机。你可以通过鼠标、键盘、手柄或触摸式设备上控制角色移动。摄像机会跟随角色，以越肩视角拍摄场景。这种视角突出主角，常用于动作及冒险类游戏。该模板的地图中还包含一些对象，可以让角色跳上跳下或四处推动。"
+			"Text": "第三人称游戏内容包包含一个可操控的角色以及一台位于角色身后斜上方的摄像机。你可以通过鼠标、键盘、手柄或触摸式设备上的虚拟摇杆控制角色移动。摄像机会跟随角色，以越肩视角拍摄场景。这种视角突出主角，常用于动作及冒险类游戏。该模板的地图中还包含一些对象，可以让角色跳上跳下或四处推动。"
 		},
 		{
 			"Language": "ja",
