Templates - TP_ThirdPerson
Updated ThirdPerson code
Implemented Combat Variant
Implemented Platforming Variant
Implemented SideScrolling Variant
#rnx

diff --git a/Templates/TP_ThirdPerson/Config/DefaultEngine.ini b/Templates/TP_ThirdPerson/Config/DefaultEngine.ini
--- a/Templates/TP_ThirdPerson/Config/DefaultEngine.ini
+++ b/Templates/TP_ThirdPerson/Config/DefaultEngine.ini
@@ -1,7 +1,7 @@
 [/Script/EngineSettings.GameMapsSettings]
-GameDefaultMap=/Game/ThirdPerson/Maps/ThirdPersonMap.ThirdPersonMap
-EditorStartupMap=/Game/ThirdPerson/Maps/ThirdPersonMap.ThirdPersonMap
-GlobalDefaultGameMode="/Script/TP_ThirdPerson.TP_ThirdPersonGameMode"
+GameDefaultMap=/Game/ThirdPerson/Lvl_ThirdPerson.Lvl_ThirdPerson
+EditorStartupMap=/Game/ThirdPerson/Lvl_ThirdPerson.Lvl_ThirdPerson
+GlobalDefaultGameMode=/Game/ThirdPerson/Blueprints/BP_ThirdPersonGameMode.BP_ThirdPersonGameMode_C
 
 [/Script/Engine.RendererSettings]
 r.ReflectionMethod=1
@@ -43,4 +43,76 @@ SoundCueCookQualityIndex=-1
 
 [/Script/LinuxTargetPlatform.LinuxTargetSettings]
 -TargetedRHIs=SF_VULKAN_SM5
-+TargetedRHIs=SF_VULKAN_SM6
\ No newline at end of file
++TargetedRHIs=SF_VULKAN_SM6
+
+[/Script/Engine.CollisionProfile]
+-Profiles=(Name="NoCollision",CollisionEnabled=NoCollision,ObjectTypeName="WorldStatic",CustomResponses=((Channel="Visibility",Response=ECR_Ignore),(Channel="Camera",Response=ECR_Ignore)),HelpMessage="No collision",bCanModify=False)
+-Profiles=(Name="BlockAll",CollisionEnabled=QueryAndPhysics,ObjectTypeName="WorldStatic",CustomResponses=,HelpMessage="WorldStatic object that blocks all actors by default. All new custom channels will use its own default response. ",bCanModify=False)
+-Profiles=(Name="OverlapAll",CollisionEnabled=QueryOnly,ObjectTypeName="WorldStatic",CustomResponses=((Channel="WorldStatic",Response=ECR_Overlap),(Channel="Pawn",Response=ECR_Overlap),(Channel="Visibility",Response=ECR_Overlap),(Channel="WorldDynamic",Response=ECR_Overlap),(Channel="Camera",Response=ECR_Overlap),(Channel="PhysicsBody",Response=ECR_Overlap),(Channel="Vehicle",Response=ECR_Overlap),(Channel="Destructible",Response=ECR_Overlap)),HelpMessage="WorldStatic object that overlaps all actors by default. All new custom channels will use its own default response. ",bCanModify=False)
+-Profiles=(Name="BlockAllDynamic",CollisionEnabled=QueryAndPhysics,ObjectTypeName="WorldDynamic",CustomResponses=,HelpMessage="WorldDynamic object that blocks all actors by default. All new custom channels will use its own default response. ",bCanModify=False)
+-Profiles=(Name="OverlapAllDynamic",CollisionEnabled=QueryOnly,ObjectTypeName="WorldDynamic",CustomResponses=((Channel="WorldStatic",Response=ECR_Overlap),(Channel="Pawn",Response=ECR_Overlap),(Channel="Visibility",Response=ECR_Overlap),(Channel="WorldDynamic",Response=ECR_Overlap),(Channel="Camera",Response=ECR_Overlap),(Channel="PhysicsBody",Response=ECR_Overlap),(Channel="Vehicle",Response=ECR_Overlap),(Channel="Destructible",Response=ECR_Overlap)),HelpMessage="WorldDynamic object that overlaps all actors by default. All new custom channels will use its own default response. ",bCanModify=False)
+-Profiles=(Name="IgnoreOnlyPawn",CollisionEnabled=QueryOnly,ObjectTypeName="WorldDynamic",CustomResponses=((Channel="Pawn",Response=ECR_Ignore),(Channel="Vehicle",Response=ECR_Ignore)),HelpMessage="WorldDynamic object that ignores Pawn and Vehicle. All other channels will be set to default.",bCanModify=False)
+-Profiles=(Name="OverlapOnlyPawn",CollisionEnabled=QueryOnly,ObjectTypeName="WorldDynamic",CustomResponses=((Channel="Pawn",Response=ECR_Overlap),(Channel="Vehicle",Response=ECR_Overlap),(Channel="Camera",Response=ECR_Ignore)),HelpMessage="WorldDynamic object that overlaps Pawn, Camera, and Vehicle. All other channels will be set to default. ",bCanModify=False)
+-Profiles=(Name="Pawn",CollisionEnabled=QueryAndPhysics,ObjectTypeName="Pawn",CustomResponses=((Channel="Visibility",Response=ECR_Ignore)),HelpMessage="Pawn object. Can be used for capsule of any playerable character or AI. ",bCanModify=False)
+-Profiles=(Name="Spectator",CollisionEnabled=QueryOnly,ObjectTypeName="Pawn",CustomResponses=((Channel="WorldStatic",Response=ECR_Block),(Channel="Pawn",Response=ECR_Ignore),(Channel="Visibility",Response=ECR_Ignore),(Channel="WorldDynamic",Response=ECR_Ignore),(Channel="Camera",Response=ECR_Ignore),(Channel="PhysicsBody",Response=ECR_Ignore),(Channel="Vehicle",Response=ECR_Ignore),(Channel="Destructible",Response=ECR_Ignore)),HelpMessage="Pawn object that ignores all other actors except WorldStatic.",bCanModify=False)
+-Profiles=(Name="CharacterMesh",CollisionEnabled=QueryOnly,ObjectTypeName="Pawn",CustomResponses=((Channel="Pawn",Response=ECR_Ignore),(Channel="Vehicle",Response=ECR_Ignore),(Channel="Visibility",Response=ECR_Ignore)),HelpMessage="Pawn object that is used for Character Mesh. All other channels will be set to default.",bCanModify=False)
+-Profiles=(Name="PhysicsActor",CollisionEnabled=QueryAndPhysics,ObjectTypeName="PhysicsBody",CustomResponses=,HelpMessage="Simulating actors",bCanModify=False)
+-Profiles=(Name="Destructible",CollisionEnabled=QueryAndPhysics,ObjectTypeName="Destructible",CustomResponses=,HelpMessage="Destructible actors",bCanModify=False)
+-Profiles=(Name="InvisibleWall",CollisionEnabled=QueryAndPhysics,ObjectTypeName="WorldStatic",CustomResponses=((Channel="Visibility",Response=ECR_Ignore)),HelpMessage="WorldStatic object that is invisible.",bCanModify=False)
+-Profiles=(Name="InvisibleWallDynamic",CollisionEnabled=QueryAndPhysics,ObjectTypeName="WorldDynamic",CustomResponses=((Channel="Visibility",Response=ECR_Ignore)),HelpMessage="WorldDynamic object that is invisible.",bCanModify=False)
+-Profiles=(Name="Trigger",CollisionEnabled=QueryOnly,ObjectTypeName="WorldDynamic",CustomResponses=((Channel="WorldStatic",Response=ECR_Overlap),(Channel="Pawn",Response=ECR_Overlap),(Channel="Visibility",Response=ECR_Ignore),(Channel="WorldDynamic",Response=ECR_Overlap),(Channel="Camera",Response=ECR_Overlap),(Channel="PhysicsBody",Response=ECR_Overlap),(Channel="Vehicle",Response=ECR_Overlap),(Channel="Destructible",Response=ECR_Overlap)),HelpMessage="WorldDynamic object that is used for trigger. All other channels will be set to default.",bCanModify=False)
+-Profiles=(Name="Ragdoll",CollisionEnabled=QueryAndPhysics,ObjectTypeName="PhysicsBody",CustomResponses=((Channel="Pawn",Response=ECR_Ignore),(Channel="Visibility",Response=ECR_Ignore)),HelpMessage="Simulating Skeletal Mesh Component. All other channels will be set to default.",bCanModify=False)
+-Profiles=(Name="Vehicle",CollisionEnabled=QueryAndPhysics,ObjectTypeName="Vehicle",CustomResponses=,HelpMessage="Vehicle object that blocks Vehicle, WorldStatic, and WorldDynamic. All other channels will be set to default.",bCanModify=False)
+-Profiles=(Name="UI",CollisionEnabled=QueryOnly,ObjectTypeName="WorldDynamic",CustomResponses=((Channel="WorldStatic",Response=ECR_Overlap),(Channel="Pawn",Response=ECR_Overlap),(Channel="Visibility",Response=ECR_Block),(Channel="WorldDynamic",Response=ECR_Overlap),(Channel="Camera",Response=ECR_Overlap),(Channel="PhysicsBody",Response=ECR_Overlap),(Channel="Vehicle",Response=ECR_Overlap),(Channel="Destructible",Response=ECR_Overlap)),HelpMessage="WorldStatic object that overlaps all actors by default. All new custom channels will use its own default response. ",bCanModify=False)
++Profiles=(Name="NoCollision",CollisionEnabled=NoCollision,bCanModify=False,ObjectTypeName="WorldStatic",CustomResponses=((Channel="Visibility",Response=ECR_Ignore),(Channel="Camera",Response=ECR_Ignore)),HelpMessage="No collision")
++Profiles=(Name="BlockAll",CollisionEnabled=QueryAndPhysics,bCanModify=False,ObjectTypeName="WorldStatic",CustomResponses=,HelpMessage="WorldStatic object that blocks all actors by default. All new custom channels will use its own default response. ")
++Profiles=(Name="OverlapAll",CollisionEnabled=QueryOnly,bCanModify=False,ObjectTypeName="WorldStatic",CustomResponses=((Channel="WorldStatic",Response=ECR_Overlap),(Channel="Pawn",Response=ECR_Overlap),(Channel="Visibility",Response=ECR_Overlap),(Channel="WorldDynamic",Response=ECR_Overlap),(Channel="Camera",Response=ECR_Overlap),(Channel="PhysicsBody",Response=ECR_Overlap),(Channel="Vehicle",Response=ECR_Overlap),(Channel="Destructible",Response=ECR_Overlap)),HelpMessage="WorldStatic object that overlaps all actors by default. All new custom channels will use its own default response. ")
++Profiles=(Name="BlockAllDynamic",CollisionEnabled=QueryAndPhysics,bCanModify=False,ObjectTypeName="WorldDynamic",CustomResponses=,HelpMessage="WorldDynamic object that blocks all actors by default. All new custom channels will use its own default response. ")
++Profiles=(Name="OverlapAllDynamic",CollisionEnabled=QueryOnly,bCanModify=False,ObjectTypeName="WorldDynamic",CustomResponses=((Channel="WorldStatic",Response=ECR_Overlap),(Channel="Pawn",Response=ECR_Overlap),(Channel="Visibility",Response=ECR_Overlap),(Channel="WorldDynamic",Response=ECR_Overlap),(Channel="Camera",Response=ECR_Overlap),(Channel="PhysicsBody",Response=ECR_Overlap),(Channel="Vehicle",Response=ECR_Overlap),(Channel="Destructible",Response=ECR_Overlap)),HelpMessage="WorldDynamic object that overlaps all actors by default. All new custom channels will use its own default response. ")
++Profiles=(Name="IgnoreOnlyPawn",CollisionEnabled=QueryOnly,bCanModify=False,ObjectTypeName="WorldDynamic",CustomResponses=((Channel="Pawn",Response=ECR_Ignore),(Channel="Vehicle",Response=ECR_Ignore)),HelpMessage="WorldDynamic object that ignores Pawn and Vehicle. All other channels will be set to default.")
++Profiles=(Name="OverlapOnlyPawn",CollisionEnabled=QueryOnly,bCanModify=False,ObjectTypeName="WorldDynamic",CustomResponses=((Channel="Pawn",Response=ECR_Overlap),(Channel="Vehicle",Response=ECR_Overlap),(Channel="Camera",Response=ECR_Ignore)),HelpMessage="WorldDynamic object that overlaps Pawn, Camera, and Vehicle. All other channels will be set to default. ")
++Profiles=(Name="Pawn",CollisionEnabled=QueryAndPhysics,bCanModify=False,ObjectTypeName="Pawn",CustomResponses=((Channel="Visibility",Response=ECR_Ignore)),HelpMessage="Pawn object. Can be used for capsule of any playerable character or AI. ")
++Profiles=(Name="Spectator",CollisionEnabled=QueryOnly,bCanModify=False,ObjectTypeName="Pawn",CustomResponses=((Channel="WorldStatic"),(Channel="Pawn",Response=ECR_Ignore),(Channel="Visibility",Response=ECR_Ignore),(Channel="WorldDynamic",Response=ECR_Ignore),(Channel="Camera",Response=ECR_Ignore),(Channel="PhysicsBody",Response=ECR_Ignore),(Channel="Vehicle",Response=ECR_Ignore),(Channel="Destructible",Response=ECR_Ignore)),HelpMessage="Pawn object that ignores all other actors except WorldStatic.")
++Profiles=(Name="CharacterMesh",CollisionEnabled=QueryOnly,bCanModify=False,ObjectTypeName="Pawn",CustomResponses=((Channel="Pawn",Response=ECR_Ignore),(Channel="Vehicle",Response=ECR_Ignore),(Channel="Visibility",Response=ECR_Ignore)),HelpMessage="Pawn object that is used for Character Mesh. All other channels will be set to default.")
++Profiles=(Name="PhysicsActor",CollisionEnabled=QueryAndPhysics,bCanModify=False,ObjectTypeName="PhysicsBody",CustomResponses=,HelpMessage="Simulating actors")
++Profiles=(Name="Destructible",CollisionEnabled=QueryAndPhysics,bCanModify=False,ObjectTypeName="Destructible",CustomResponses=,HelpMessage="Destructible actors")
++Profiles=(Name="InvisibleWall",CollisionEnabled=QueryAndPhysics,bCanModify=False,ObjectTypeName="WorldStatic",CustomResponses=((Channel="Visibility",Response=ECR_Ignore)),HelpMessage="WorldStatic object that is invisible.")
++Profiles=(Name="InvisibleWallDynamic",CollisionEnabled=QueryAndPhysics,bCanModify=False,ObjectTypeName="WorldDynamic",CustomResponses=((Channel="Visibility",Response=ECR_Ignore)),HelpMessage="WorldDynamic object that is invisible.")
++Profiles=(Name="Trigger",CollisionEnabled=QueryOnly,bCanModify=False,ObjectTypeName="WorldDynamic",CustomResponses=((Channel="WorldStatic",Response=ECR_Overlap),(Channel="Pawn",Response=ECR_Overlap),(Channel="Visibility",Response=ECR_Ignore),(Channel="WorldDynamic",Response=ECR_Overlap),(Channel="Camera",Response=ECR_Overlap),(Channel="PhysicsBody",Response=ECR_Overlap),(Channel="Vehicle",Response=ECR_Overlap),(Channel="Destructible",Response=ECR_Overlap)),HelpMessage="WorldDynamic object that is used for trigger. All other channels will be set to default.")
++Profiles=(Name="Ragdoll",CollisionEnabled=QueryAndPhysics,bCanModify=False,ObjectTypeName="PhysicsBody",CustomResponses=((Channel="Pawn",Response=ECR_Ignore),(Channel="Visibility",Response=ECR_Ignore)),HelpMessage="Simulating Skeletal Mesh Component. All other channels will be set to default.")
++Profiles=(Name="Vehicle",CollisionEnabled=QueryAndPhysics,bCanModify=False,ObjectTypeName="Vehicle",CustomResponses=,HelpMessage="Vehicle object that blocks Vehicle, WorldStatic, and WorldDynamic. All other channels will be set to default.")
++Profiles=(Name="UI",CollisionEnabled=QueryOnly,bCanModify=False,ObjectTypeName="WorldDynamic",CustomResponses=((Channel="WorldStatic",Response=ECR_Overlap),(Channel="Pawn",Response=ECR_Overlap),(Channel="Visibility"),(Channel="WorldDynamic",Response=ECR_Overlap),(Channel="Camera",Response=ECR_Overlap),(Channel="PhysicsBody",Response=ECR_Overlap),(Channel="Vehicle",Response=ECR_Overlap),(Channel="Destructible",Response=ECR_Overlap)),HelpMessage="WorldStatic object that overlaps all actors by default. All new custom channels will use its own default response. ")
++DefaultChannelResponses=(Channel=ECC_GameTraceChannel1,DefaultResponse=ECR_Block,bTraceType=False,bStaticObject=False,Name="SoftCollision")
+-ProfileRedirects=(OldName="BlockingVolume",NewName="InvisibleWall")
+-ProfileRedirects=(OldName="InterpActor",NewName="IgnoreOnlyPawn")
+-ProfileRedirects=(OldName="StaticMeshComponent",NewName="BlockAllDynamic")
+-ProfileRedirects=(OldName="SkeletalMeshActor",NewName="PhysicsActor")
+-ProfileRedirects=(OldName="InvisibleActor",NewName="InvisibleWallDynamic")
++ProfileRedirects=(OldName="BlockingVolume",NewName="InvisibleWall")
++ProfileRedirects=(OldName="InterpActor",NewName="IgnoreOnlyPawn")
++ProfileRedirects=(OldName="StaticMeshComponent",NewName="BlockAllDynamic")
++ProfileRedirects=(OldName="SkeletalMeshActor",NewName="PhysicsActor")
++ProfileRedirects=(OldName="InvisibleActor",NewName="InvisibleWallDynamic")
+-CollisionChannelRedirects=(OldName="Static",NewName="WorldStatic")
+-CollisionChannelRedirects=(OldName="Dynamic",NewName="WorldDynamic")
+-CollisionChannelRedirects=(OldName="VehicleMovement",NewName="Vehicle")
+-CollisionChannelRedirects=(OldName="PawnMovement",NewName="Pawn")
++CollisionChannelRedirects=(OldName="Static",NewName="WorldStatic")
++CollisionChannelRedirects=(OldName="Dynamic",NewName="WorldDynamic")
++CollisionChannelRedirects=(OldName="VehicleMovement",NewName="Vehicle")
++CollisionChannelRedirects=(OldName="PawnMovement",NewName="Pawn")
+
+[/Script/AndroidFileServerEditor.AndroidFileServerRuntimeSettings]
+bEnablePlugin=True
+bAllowNetworkConnection=True
+SecurityToken=F022329C4DC0428B5B294D852830B9FD
+bIncludeInShipping=False
+bAllowExternalStartInShipping=False
+bCompileAFSProject=False
+bUseCompression=False
+bLogFiles=False
+bReportStats=False
+ConnectionType=USBOnly
+bUseManualIPAddress=False
+ManualIPAddress=
+
diff --git a/Templates/TP_ThirdPerson/Config/TemplateDefs.ini b/Templates/TP_ThirdPerson/Config/TemplateDefs.ini
--- a/Templates/TP_ThirdPerson/Config/TemplateDefs.ini
+++ b/Templates/TP_ThirdPerson/Config/TemplateDefs.ini
@@ -1,4 +1,4 @@
-﻿[/Script/GameProjectGeneration.TemplateProjectDefs]
+[/Script/GameProjectGeneration.TemplateProjectDefs]
 
 Categories=Games
 
@@ -53,4 +53,5 @@ ReplacementsInFiles=(Extensions=("cpp","h","ini","cs"),From="%TEMPLATENAME_LOWER
 ReplacementsInFiles=(Extensions=("cpp","h","ini","cs"),From="%TEMPLATENAME%",To="%PROJECTNAME%",bCaseSensitive=false)
 SharedContentPacks=(MountName="LevelPrototyping",DetailLevels=("High"))
 SharedContentPacks=(MountName="Characters",DetailLevels=("High"))
+SharedContentPacks=(MountName="Input",DetailLevels=("High"))
 EditDetailLevelPreference="High"
\ No newline at end of file
diff --git a/Templates/TP_ThirdPerson/Source/TP_ThirdPerson/TP_ThirdPerson.Build.cs b/Templates/TP_ThirdPerson/Source/TP_ThirdPerson/TP_ThirdPerson.Build.cs
--- a/Templates/TP_ThirdPerson/Source/TP_ThirdPerson/TP_ThirdPerson.Build.cs
+++ b/Templates/TP_ThirdPerson/Source/TP_ThirdPerson/TP_ThirdPerson.Build.cs
@@ -8,6 +8,36 @@ public class TP_ThirdPerson : ModuleRules
 	{
 		PCHUsage = PCHUsageMode.UseExplicitOrSharedPCHs;
 
-		PublicDependencyModuleNames.AddRange(new string[] { "Core", "CoreUObject", "Engine", "InputCore", "EnhancedInput" });
+		PublicDependencyModuleNames.AddRange(new string[] {
+			"Core",
+			"CoreUObject",
+			"Engine",
+			"InputCore",
+			"EnhancedInput",
+			"AIModule",
+			"StateTreeModule",
+			"GameplayStateTreeModule",
+			"UMG"
+		});
+
+		PrivateDependencyModuleNames.AddRange(new string[] { });
+
+		PublicIncludePaths.AddRange(new string[] {
+			"TP_ThirdPerson",
+			"TP_ThirdPerson/Variant_Platforming",
+			"TP_ThirdPerson/Variant_Combat",
+			"TP_ThirdPerson/Variant_Combat/AI",
+			"TP_ThirdPerson/Variant_SideScrolling",
+			"TP_ThirdPerson/Variant_SideScrolling/Gameplay",
+			"TP_ThirdPerson/Variant_SideScrolling/AI"
+		});
+
+		// Uncomment if you are using Slate UI
+		// PrivateDependencyModuleNames.AddRange(new string[] { "Slate", "SlateCore" });
+
+		// Uncomment if you are using online features
+		// PrivateDependencyModuleNames.Add("OnlineSubsystem");
+
+		// To include OnlineSubsystemSteam, add it to the plugins section in your uproject file with the Enabled attribute set to true
 	}
 }
diff --git a/Templates/TP_ThirdPerson/Source/TP_ThirdPerson/TP_ThirdPersonCharacter.cpp b/Templates/TP_ThirdPerson/Source/TP_ThirdPerson/TP_ThirdPersonCharacter.cpp
--- a/Templates/TP_ThirdPerson/Source/TP_ThirdPerson/TP_ThirdPersonCharacter.cpp
+++ b/Templates/TP_ThirdPerson/Source/TP_ThirdPerson/TP_ThirdPersonCharacter.cpp
@@ -13,9 +13,6 @@
 
 DEFINE_LOG_CATEGORY(LogTemplateCharacter);
 
-//////////////////////////////////////////////////////////////////////////
-// ATP_ThirdPersonCharacter
-
 ATP_ThirdPersonCharacter::ATP_ThirdPersonCharacter()
 {
 	// Set size for collision capsule
@@ -27,8 +24,8 @@ ATP_ThirdPersonCharacter::ATP_ThirdPersonCharacter()
 	bUseControllerRotationRoll = false;
 
 	// Configure character movement
-	GetCharacterMovement()->bOrientRotationToMovement = true; // Character moves in the direction of input...	
-	GetCharacterMovement()->RotationRate = FRotator(0.0f, 500.0f, 0.0f); // ...at this rotation rate
+	GetCharacterMovement()->bOrientRotationToMovement = true;
+	GetCharacterMovement()->RotationRate = FRotator(0.0f, 500.0f, 0.0f);
 
 	// Note: For faster iteration times these variables, and many more, can be tweaked in the Character Blueprint
 	// instead of recompiling to adjust them
@@ -42,35 +39,18 @@ ATP_ThirdPersonCharacter::ATP_ThirdPersonCharacter()
 	// Create a camera boom (pulls in towards the player if there is a collision)
 	CameraBoom = CreateDefaultSubobject<USpringArmComponent>(TEXT("CameraBoom"));
 	CameraBoom->SetupAttachment(RootComponent);
-	CameraBoom->TargetArmLength = 400.0f; // The camera follows at this distance behind the character	
-	CameraBoom->bUsePawnControlRotation = true; // Rotate the arm based on the controller
+	CameraBoom->TargetArmLength = 400.0f;
+	CameraBoom->bUsePawnControlRotation = true;
 
 	// Create a follow camera
 	FollowCamera = CreateDefaultSubobject<UCameraComponent>(TEXT("FollowCamera"));
-	FollowCamera->SetupAttachment(CameraBoom, USpringArmComponent::SocketName); // Attach the camera to the end of the boom and let the boom adjust to match the controller orientation
-	FollowCamera->bUsePawnControlRotation = false; // Camera does not rotate relative to arm
+	FollowCamera->SetupAttachment(CameraBoom, USpringArmComponent::SocketName);
+	FollowCamera->bUsePawnControlRotation = false;
 
 	// Note: The skeletal mesh and anim blueprint references on the Mesh component (inherited from Character) 
 	// are set in the derived blueprint asset named ThirdPersonCharacter (to avoid direct content references in C++)
 }
 
-//////////////////////////////////////////////////////////////////////////
-// Input
-
-void ATP_ThirdPersonCharacter::NotifyControllerChanged()
-{
-	Super::NotifyControllerChanged();
-
-	// Add Input Mapping Context
-	if (APlayerController* PlayerController = Cast<APlayerController>(GetController()))
-	{
-		if (UEnhancedInputLocalPlayerSubsystem* Subsystem = ULocalPlayer::GetSubsystem<UEnhancedInputLocalPlayerSubsystem>(PlayerController->GetLocalPlayer()))
-		{
-			Subsystem->AddMappingContext(DefaultMappingContext, 0);
-		}
-	}
-}
-
 void ATP_ThirdPersonCharacter::SetupPlayerInputComponent(UInputComponent* PlayerInputComponent)
 {
 	// Set up action bindings
@@ -82,6 +62,7 @@ void ATP_ThirdPersonCharacter::SetupPlayerInputComponent(UInputComponent* Player
 
 		// Moving
 		EnhancedInputComponent->BindAction(MoveAction, ETriggerEvent::Triggered, this, &ATP_ThirdPersonCharacter::Move);
+		EnhancedInputComponent->BindAction(MouseLookAction, ETriggerEvent::Triggered, this, &ATP_ThirdPersonCharacter::Look);
 
 		// Looking
 		EnhancedInputComponent->BindAction(LookAction, ETriggerEvent::Triggered, this, &ATP_ThirdPersonCharacter::Look);
@@ -97,6 +78,21 @@ void ATP_ThirdPersonCharacter::Move(const FInputActionValue& Value)
 	// input is a Vector2D
 	FVector2D MovementVector = Value.Get<FVector2D>();
 
+	// route the input
+	DoMove(MovementVector.X, MovementVector.Y);
+}
+
+void ATP_ThirdPersonCharacter::Look(const FInputActionValue& Value)
+{
+	// input is a Vector2D
+	FVector2D LookAxisVector = Value.Get<FVector2D>();
+
+	// route the input
+	DoLook(LookAxisVector.X, LookAxisVector.Y);
+}
+
+void ATP_ThirdPersonCharacter::DoMove(float Right, float Forward)
+{
 	if (GetController() != nullptr)
 	{
 		// find out which way is forward
@@ -105,25 +101,34 @@ void ATP_ThirdPersonCharacter::Move(const FInputActionValue& Value)
 
 		// get forward vector
 		const FVector ForwardDirection = FRotationMatrix(YawRotation).GetUnitAxis(EAxis::X);
-	
+
 		// get right vector 
 		const FVector RightDirection = FRotationMatrix(YawRotation).GetUnitAxis(EAxis::Y);
 
 		// add movement 
-		AddMovementInput(ForwardDirection, MovementVector.Y);
-		AddMovementInput(RightDirection, MovementVector.X);
+		AddMovementInput(ForwardDirection, Forward);
+		AddMovementInput(RightDirection, Right);
 	}
 }
 
-void ATP_ThirdPersonCharacter::Look(const FInputActionValue& Value)
+void ATP_ThirdPersonCharacter::DoLook(float Yaw, float Pitch)
 {
-	// input is a Vector2D
-	FVector2D LookAxisVector = Value.Get<FVector2D>();
-
 	if (GetController() != nullptr)
 	{
 		// add yaw and pitch input to controller
-		AddControllerYawInput(LookAxisVector.X);
-		AddControllerPitchInput(LookAxisVector.Y);
+		AddControllerYawInput(Yaw);
+		AddControllerPitchInput(Pitch);
 	}
 }
+
+void ATP_ThirdPersonCharacter::DoJumpStart()
+{
+	// signal the character to jump
+	Jump();
+}
+
+void ATP_ThirdPersonCharacter::DoJumpEnd()
+{
+	// signal the character to stop jumping
+	StopJumping();
+}
diff --git a/Templates/TP_ThirdPerson/Source/TP_ThirdPerson/TP_ThirdPersonCharacter.h b/Templates/TP_ThirdPerson/Source/TP_ThirdPerson/TP_ThirdPersonCharacter.h
--- a/Templates/TP_ThirdPerson/Source/TP_ThirdPerson/TP_ThirdPersonCharacter.h
+++ b/Templates/TP_ThirdPerson/Source/TP_ThirdPerson/TP_ThirdPersonCharacter.h
@@ -9,13 +9,16 @@
 
 class USpringArmComponent;
 class UCameraComponent;
-class UInputMappingContext;
 class UInputAction;
 struct FInputActionValue;
 
 DECLARE_LOG_CATEGORY_EXTERN(LogTemplateCharacter, Log, All);
 
-UCLASS(config=Game)
+/**
+ *  A simple player-controllable third person character
+ *  Implements a controllable orbiting camera
+ */
+UCLASS(abstract)
 class ATP_ThirdPersonCharacter : public ACharacter
 {
 	GENERATED_BODY()
@@ -28,25 +31,33 @@ class ATP_ThirdPersonCharacter : public ACharacter
 	UPROPERTY(VisibleAnywhere, BlueprintReadOnly, Category = Camera, meta = (AllowPrivateAccess = "true"))
 	UCameraComponent* FollowCamera;
 	
-	/** MappingContext */
-	UPROPERTY(EditAnywhere, BlueprintReadOnly, Category = Input, meta = (AllowPrivateAccess = "true"))
-	UInputMappingContext* DefaultMappingContext;
+protected:
 
 	/** Jump Input Action */
-	UPROPERTY(EditAnywhere, BlueprintReadOnly, Category = Input, meta = (AllowPrivateAccess = "true"))
+	UPROPERTY(EditAnywhere, BlueprintReadOnly, Category="Input")
 	UInputAction* JumpAction;
 
 	/** Move Input Action */
-	UPROPERTY(EditAnywhere, BlueprintReadOnly, Category = Input, meta = (AllowPrivateAccess = "true"))
+	UPROPERTY(EditAnywhere, BlueprintReadOnly, Category="Input")
 	UInputAction* MoveAction;
 
 	/** Look Input Action */
-	UPROPERTY(EditAnywhere, BlueprintReadOnly, Category = Input, meta = (AllowPrivateAccess = "true"))
+	UPROPERTY(EditAnywhere, BlueprintReadOnly, Category="Input")
 	UInputAction* LookAction;
 
+	/** Mouse Look Input Action */
+	UPROPERTY(EditAnywhere, BlueprintReadOnly, Category="Input")
+	UInputAction* MouseLookAction;
+
 public:
-	ATP_ThirdPersonCharacter();
-	
+
+	/** Constructor */
+	ATP_ThirdPersonCharacter();	
+
+protected:
+
+	/** Initialize input action bindings */
+	virtual void SetupPlayerInputComponent(class UInputComponent* PlayerInputComponent) override;
 
 protected:
 
@@ -55,17 +66,30 @@ protected:
 
 	/** Called for looking input */
 	void Look(const FInputActionValue& Value);
-			
 
-protected:
+public:
 
-	virtual void NotifyControllerChanged() override;
+	/** Handles move inputs from either controls or UI interfaces */
+	UFUNCTION(BlueprintCallable, Category="Input")
+	virtual void DoMove(float Right, float Forward);
 
-	virtual void SetupPlayerInputComponent(class UInputComponent* PlayerInputComponent) override;
+	/** Handles look inputs from either controls or UI interfaces */
+	UFUNCTION(BlueprintCallable, Category="Input")
+	virtual void DoLook(float Yaw, float Pitch);
+
+	/** Handles jump pressed inputs from either controls or UI interfaces */
+	UFUNCTION(BlueprintCallable, Category="Input")
+	virtual void DoJumpStart();
+
+	/** Handles jump pressed inputs from either controls or UI interfaces */
+	UFUNCTION(BlueprintCallable, Category="Input")
+	virtual void DoJumpEnd();
 
 public:
+
 	/** Returns CameraBoom subobject **/
 	FORCEINLINE class USpringArmComponent* GetCameraBoom() const { return CameraBoom; }
+
 	/** Returns FollowCamera subobject **/
 	FORCEINLINE class UCameraComponent* GetFollowCamera() const { return FollowCamera; }
 };
diff --git a/Templates/TP_ThirdPerson/Source/TP_ThirdPerson/TP_ThirdPersonGameMode.cpp b/Templates/TP_ThirdPerson/Source/TP_ThirdPerson/TP_ThirdPersonGameMode.cpp
--- a/Templates/TP_ThirdPerson/Source/TP_ThirdPerson/TP_ThirdPersonGameMode.cpp
+++ b/Templates/TP_ThirdPerson/Source/TP_ThirdPerson/TP_ThirdPersonGameMode.cpp
@@ -1,15 +1,8 @@
 // Copyright Epic Games, Inc. All Rights Reserved.
 
 #include "TP_ThirdPersonGameMode.h"
-#include "TP_ThirdPersonCharacter.h"
-#include "UObject/ConstructorHelpers.h"
 
 ATP_ThirdPersonGameMode::ATP_ThirdPersonGameMode()
 {
-	// set default pawn class to our Blueprinted character
-	static ConstructorHelpers::FClassFinder<APawn> PlayerPawnBPClass(TEXT("/Game/ThirdPerson/Blueprints/BP_ThirdPersonCharacter"));
-	if (PlayerPawnBPClass.Class != NULL)
-	{
-		DefaultPawnClass = PlayerPawnBPClass.Class;
-	}
+	// stub
 }
diff --git a/Templates/TP_ThirdPerson/Source/TP_ThirdPerson/TP_ThirdPersonGameMode.h b/Templates/TP_ThirdPerson/Source/TP_ThirdPerson/TP_ThirdPersonGameMode.h
--- a/Templates/TP_ThirdPerson/Source/TP_ThirdPerson/TP_ThirdPersonGameMode.h
+++ b/Templates/TP_ThirdPerson/Source/TP_ThirdPerson/TP_ThirdPersonGameMode.h
@@ -6,12 +6,17 @@
 #include "GameFramework/GameModeBase.h"
 #include "TP_ThirdPersonGameMode.generated.h"
 
-UCLASS(minimalapi)
+/**
+ *  Simple GameMode for a third person game
+ */
+UCLASS(abstract)
 class ATP_ThirdPersonGameMode : public AGameModeBase
 {
 	GENERATED_BODY()
 
 public:
+	
+	/** Constructor */
 	ATP_ThirdPersonGameMode();
 };
 
diff --git a/Templates/TP_ThirdPerson/Source/TP_ThirdPerson/TP_ThirdPersonPlayerController.cpp b/Templates/TP_ThirdPerson/Source/TP_ThirdPerson/TP_ThirdPersonPlayerController.cpp
new file mode 100644
--- /dev/null
+++ b/Templates/TP_ThirdPerson/Source/TP_ThirdPerson/TP_ThirdPersonPlayerController.cpp
@@ -0,0 +1,21 @@
+// Copyright Epic Games, Inc. All Rights Reserved.
+
+
+#include "TP_ThirdPersonPlayerController.h"
+#include "EnhancedInputSubsystems.h"
+#include "Engine/LocalPlayer.h"
+#include "InputMappingContext.h"
+
+void ATP_ThirdPersonPlayerController::SetupInputComponent()
+{
+	Super::SetupInputComponent();
+
+	// Add Input Mapping Contexts
+	if (UEnhancedInputLocalPlayerSubsystem* Subsystem = ULocalPlayer::GetSubsystem<UEnhancedInputLocalPlayerSubsystem>(GetLocalPlayer()))
+	{
+		for (UInputMappingContext* CurrentContext : DefaultMappingContexts)
+		{
+			Subsystem->AddMappingContext(CurrentContext, 0);
+		}
+	}
+}
diff --git a/Templates/TP_ThirdPerson/Source/TP_ThirdPerson/TP_ThirdPersonPlayerController.h b/Templates/TP_ThirdPerson/Source/TP_ThirdPerson/TP_ThirdPersonPlayerController.h
new file mode 100644
--- /dev/null
+++ b/Templates/TP_ThirdPerson/Source/TP_ThirdPerson/TP_ThirdPersonPlayerController.h
@@ -0,0 +1,29 @@
+// Copyright Epic Games, Inc. All Rights Reserved.
+
+#pragma once
+
+#include "CoreMinimal.h"
+#include "GameFramework/PlayerController.h"
+#include "TP_ThirdPersonPlayerController.generated.h"
+
+class UInputMappingContext;
+
+/**
+ *  Basic PlayerController class for a third person game
+ *  Manages input mappings
+ */
+UCLASS(abstract)
+class ATP_ThirdPersonPlayerController : public APlayerController
+{
+	GENERATED_BODY()
+	
+protected:
+
+	/** Input Mapping Contexts */
+	UPROPERTY(EditAnywhere, BlueprintReadOnly, Category ="Input", meta = (AllowPrivateAccess = "true"))
+	TArray<UInputMappingContext*> DefaultMappingContexts;
+
+	/** Input mapping context setup */
+	virtual void SetupInputComponent() override;
+
+};
diff --git a/Templates/TP_ThirdPerson/Source/TP_ThirdPerson/Variant_Combat/AI/CombatAIController.cpp b/Templates/TP_ThirdPerson/Source/TP_ThirdPerson/Variant_Combat/AI/CombatAIController.cpp
new file mode 100644
--- /dev/null
+++ b/Templates/TP_ThirdPerson/Source/TP_ThirdPerson/Variant_Combat/AI/CombatAIController.cpp
@@ -0,0 +1,19 @@
+// Copyright Epic Games, Inc. All Rights Reserved.
+
+
+#include "CombatAIController.h"
+#include "Components/StateTreeAIComponent.h"
+
+ACombatAIController::ACombatAIController()
+{
+	// create the StateTree AI Component
+	StateTreeAI = CreateDefaultSubobject<UStateTreeAIComponent>(TEXT("StateTreeAI"));
+	check(StateTreeAI);
+
+	// ensure we start the StateTree when we possess the pawn
+	bStartAILogicOnPossess = true;
+
+	// ensure we're attached to the possessed character.
+	// this is necessary for EnvQueries to work correctly
+	bAttachToPawn = true;
+}
diff --git a/Templates/TP_ThirdPerson/Source/TP_ThirdPerson/Variant_Combat/AI/CombatAIController.h b/Templates/TP_ThirdPerson/Source/TP_ThirdPerson/Variant_Combat/AI/CombatAIController.h
new file mode 100644
--- /dev/null
+++ b/Templates/TP_ThirdPerson/Source/TP_ThirdPerson/Variant_Combat/AI/CombatAIController.h
@@ -0,0 +1,27 @@
+// Copyright Epic Games, Inc. All Rights Reserved.
+
+#pragma once
+
+#include "CoreMinimal.h"
+#include "AIController.h"
+#include "CombatAIController.generated.h"
+
+class UStateTreeAIComponent;
+
+/**
+ *	A basic AI Controller capable of running StateTree
+ */
+UCLASS(abstract)
+class ACombatAIController : public AAIController
+{
+	GENERATED_BODY()
+
+	/** StateTree Component */
+	UPROPERTY(VisibleAnywhere, BlueprintReadOnly, Category = "AI", meta = (AllowPrivateAccess = "true"))
+	UStateTreeAIComponent* StateTreeAI;
+
+public:
+
+	/** Constructor */
+	ACombatAIController();
+};
diff --git a/Templates/TP_ThirdPerson/Source/TP_ThirdPerson/Variant_Combat/AI/CombatEnemy.cpp b/Templates/TP_ThirdPerson/Source/TP_ThirdPerson/Variant_Combat/AI/CombatEnemy.cpp
new file mode 100644
--- /dev/null
+++ b/Templates/TP_ThirdPerson/Source/TP_ThirdPerson/Variant_Combat/AI/CombatEnemy.cpp
@@ -0,0 +1,322 @@
+// Copyright Epic Games, Inc. All Rights Reserved.
+
+
+#include "CombatEnemy.h"
+#include "Components/CapsuleComponent.h"
+#include "GameFramework/CharacterMovementComponent.h"
+#include "CombatAIController.h"
+#include "Components/WidgetComponent.h"
+#include "Engine/DamageEvents.h"
+#include "CombatLifeBar.h"
+#include "TimerManager.h"
+#include "Components/SkeletalMeshComponent.h"
+#include "Animation/AnimInstance.h"
+
+ACombatEnemy::ACombatEnemy()
+{
+	PrimaryActorTick.bCanEverTick = true;
+
+	// bind the attack montage ended delegate
+	OnAttackMontageEnded.BindUObject(this, &ACombatEnemy::AttackMontageEnded);
+
+	// set the AI Controller class by default
+	AIControllerClass = ACombatAIController::StaticClass();
+
+	// use an AI Controller regardless of whether we're placed or spawned
+	AutoPossessAI = EAutoPossessAI::PlacedInWorldOrSpawned;
+
+	// ignore the controller's yaw rotation
+	bUseControllerRotationYaw = false;
+
+	// create the life bar
+	LifeBar = CreateDefaultSubobject<UWidgetComponent>(TEXT("LifeBar"));
+	LifeBar->SetupAttachment(RootComponent);
+
+	// set the collision capsule size
+	GetCapsuleComponent()->SetCapsuleSize(35.0f, 90.0f);
+
+	// set the character movement properties
+	GetCharacterMovement()->bUseControllerDesiredRotation = true;
+
+	// reset HP to maximum
+	CurrentHP = MaxHP;
+}
+
+void ACombatEnemy::DoAIComboAttack()
+{
+	// ignore if we're already playing an attack animation
+	if (bIsAttacking)
+	{
+		return;
+	}
+
+	// raise the attacking flag
+	bIsAttacking = true;
+
+	// choose how many times we're going to attack
+	TargetComboCount = FMath::RandRange(1, ComboSectionNames.Num() - 1);
+
+	// reset the attack counter
+	CurrentComboAttack = 0;
+
+	// play the attack montage
+	if (UAnimInstance* AnimInstance = GetMesh()->GetAnimInstance())
+	{
+		const float MontageLength = AnimInstance->Montage_Play(ComboAttackMontage, 1.0f, EMontagePlayReturnType::MontageLength, 0.0f, true);
+
+		// subscribe to montage completed and interrupted events
+		if (MontageLength > 0.0f)
+		{
+			// set the end delegate for the montage
+			AnimInstance->Montage_SetEndDelegate(OnAttackMontageEnded, ComboAttackMontage);
+		}
+	}
+}
+
+void ACombatEnemy::DoAIChargedAttack()
+{
+	// ignore if we're already playing an attack animation
+	if (bIsAttacking)
+	{
+		return;
+	}
+
+	// raise the attacking flag
+	bIsAttacking = true;
+
+	// choose how many loops are we going to charge for
+	TargetChargeLoops = FMath::RandRange(MinChargeLoops, MaxChargeLoops);
+
+	// reset the charge loop counter
+	CurrentChargeLoop = 0;
+
+	// play the attack montage
+	if (UAnimInstance* AnimInstance = GetMesh()->GetAnimInstance())
+	{
+		const float MontageLength = AnimInstance->Montage_Play(ChargedAttackMontage, 1.0f, EMontagePlayReturnType::MontageLength, 0.0f, true);
+
+		// subscribe to montage completed and interrupted events
+		if (MontageLength > 0.0f)
+		{
+			// set the end delegate for the montage
+			AnimInstance->Montage_SetEndDelegate(OnAttackMontageEnded, ChargedAttackMontage);
+		}
+	}
+}
+
+void ACombatEnemy::AttackMontageEnded(UAnimMontage* Montage, bool bInterrupted)
+{
+	// reset the attacking flag
+	bIsAttacking = false;
+
+	// call the attack completed delegate so the StateTree can continue execution
+	OnAttackCompleted.ExecuteIfBound();
+}
+
+void ACombatEnemy::DoAttackTrace(FName DamageSourceBone)
+{
+	// sweep for objects in front of the character to be hit by the attack
+	TArray<FHitResult> OutHits;
+
+	// start at the provided socket location, sweep forward
+	const FVector TraceStart = GetMesh()->GetSocketLocation(DamageSourceBone);
+	const FVector TraceEnd = TraceStart + (GetActorForwardVector() * MeleeTraceDistance);
+
+	// enemies only affect Pawn collision objects; they don't knock back boxes
+	FCollisionObjectQueryParams ObjectParams;
+	ObjectParams.AddObjectTypesToQuery(ECC_Pawn);
+
+	// use a sphere shape for the sweep
+	FCollisionShape CollisionShape;
+	CollisionShape.SetSphere(MeleeTraceRadius);
+
+	// ignore self
+	FCollisionQueryParams QueryParams;
+	QueryParams.AddIgnoredActor(this);
+
+	if (GetWorld()->SweepMultiByObjectType(OutHits, TraceStart, TraceEnd, FQuat::Identity, ObjectParams, CollisionShape, QueryParams))
+	{
+		// iterate over each object hit
+		for (const FHitResult& CurrentHit : OutHits)
+		{
+			/** does the actor have the player tag? */
+			if (CurrentHit.GetActor()->ActorHasTag(FName("Player")))
+			{
+				// check if the actor is damageable
+				ICombatDamageable* Damageable = Cast<ICombatDamageable>(CurrentHit.GetActor());
+
+				if (Damageable)
+				{
+					// knock upwards and away from the impact normal
+					const FVector Impulse = (CurrentHit.ImpactNormal * -MeleeKnockbackImpulse) + (FVector::UpVector * MeleeLaunchImpulse);
+
+					// pass the damage event to the actor
+					Damageable->ApplyDamage(MeleeDamage, this, CurrentHit.ImpactPoint, Impulse);
+
+				}
+			}
+		}
+	}
+}
+
+void ACombatEnemy::CheckCombo()
+{
+	// increase the combo counter
+	++CurrentComboAttack;
+
+	// do we still have attacks to play in this string?
+	if (CurrentComboAttack < TargetComboCount)
+	{
+		// jump to the next attack section
+		if (UAnimInstance* AnimInstance = GetMesh()->GetAnimInstance())
+		{
+			AnimInstance->Montage_JumpToSection(ComboSectionNames[CurrentComboAttack], ComboAttackMontage);
+		}
+	}
+}
+
+void ACombatEnemy::CheckChargedAttack()
+{
+	// increase the charge loop counter
+	++CurrentChargeLoop;
+
+	// jump to either the loop or attack section of the montage depending on whether we hit the loop target
+	if (UAnimInstance* AnimInstance = GetMesh()->GetAnimInstance())
+	{
+		AnimInstance->Montage_JumpToSection(CurrentChargeLoop >= TargetChargeLoops ? ChargeAttackSection : ChargeLoopSection, ChargedAttackMontage);
+	}
+}
+
+void ACombatEnemy::ApplyDamage(float Damage, AActor* DamageCauser, const FVector& DamageLocation, const FVector& DamageImpulse)
+{
+	
+	// pass the damage event to the actor
+	FDamageEvent DamageEvent;
+	const float ActualDamage = TakeDamage(Damage, DamageEvent, nullptr, DamageCauser);
+
+	// only process knockback and effects if we received nonzero damage
+	if (ActualDamage > 0.0f)
+	{
+		// apply the knockback impulse
+		GetCharacterMovement()->AddImpulse(DamageImpulse, true);
+
+		// is the character ragdolling?
+		if (GetMesh()->IsSimulatingPhysics())
+		{
+			// apply an impulse to the ragdoll
+			GetMesh()->AddImpulseAtLocation(DamageImpulse * GetMesh()->GetMass(), DamageLocation);
+		}
+
+		// stop the attack montages to interrupt the attack
+		if (UAnimInstance* AnimInstance = GetMesh()->GetAnimInstance())
+		{
+			AnimInstance->Montage_Stop(0.1f, ComboAttackMontage);
+			AnimInstance->Montage_Stop(0.1f, ChargedAttackMontage);
+		}
+
+		// pass control to BP to play effects, etc.
+		ReceivedDamage(ActualDamage, DamageLocation, DamageImpulse.GetSafeNormal());
+	}
+}
+
+void ACombatEnemy::HandleDeath()
+{
+	// hide the life bar
+	LifeBar->SetHiddenInGame(true);
+
+	// disable the collision capsule to avoid being hit again while dead
+	GetCapsuleComponent()->SetCollisionEnabled(ECollisionEnabled::NoCollision);
+
+	// disable character movement
+	GetCharacterMovement()->DisableMovement();
+
+	// enable full ragdoll physics
+	GetMesh()->SetSimulatePhysics(true);
+
+	// call the died delegate to notify any subscribers
+	OnEnemyDied.Broadcast();
+
+	// set up the death timer
+	GetWorld()->GetTimerManager().SetTimer(DeathTimer, this, &ACombatEnemy::RemoveFromLevel, DeathRemovalTime);
+}
+
+void ACombatEnemy::ApplyHealing(float Healing, AActor* Healer)
+{
+	// stub
+}
+
+void ACombatEnemy::RemoveFromLevel()
+{
+	// destroy this actor
+	Destroy();
+}
+
+float ACombatEnemy::TakeDamage(float Damage, struct FDamageEvent const& DamageEvent, AController* EventInstigator, AActor* DamageCauser)
+{
+	// only process damage if the character is still alive
+	if (CurrentHP <= 0.0f)
+	{
+		return 0.0f;
+	}
+
+	// reduce the current HP
+	CurrentHP -= Damage;
+
+	// have we run out of HP?
+	if (CurrentHP <= 0.0f)
+	{
+		// die
+		HandleDeath();
+	}
+	else
+	{
+		// update the life bar
+		LifeBarWidget->SetLifePercentage(CurrentHP / MaxHP);
+
+		// enable partial ragdoll physics, but keep the pelvis vertical
+		GetMesh()->SetPhysicsBlendWeight(0.5f);
+		GetMesh()->SetBodySimulatePhysics(PelvisBoneName, false);
+	}
+
+	// return the received damage amount
+	return Damage;
+}
+
+void ACombatEnemy::Landed(const FHitResult& Hit)
+{
+	Super::Landed(Hit);
+
+	// is the character still alive?
+	if (CurrentHP >= 0.0f)
+	{
+		// disable ragdoll physics
+		GetMesh()->SetPhysicsBlendWeight(0.0f);
+	}
+
+	// call the landed Delegate for StateTree
+	OnEnemyLanded.ExecuteIfBound();
+}
+
+void ACombatEnemy::BeginPlay()
+{
+	// reset HP to maximum
+	CurrentHP = MaxHP;
+
+	// we top the HP before BeginPlay so StateTree picks it up at the right value
+	Super::BeginPlay();
+
+	// get the life bar widget from the widget comp
+	LifeBarWidget = Cast<UCombatLifeBar>(LifeBar->GetUserWidgetObject());
+	check(LifeBarWidget);
+
+	// fill the life bar
+	LifeBarWidget->SetLifePercentage(1.0f);
+}
+
+void ACombatEnemy::EndPlay(EEndPlayReason::Type EndPlayReason)
+{
+	Super::EndPlay(EndPlayReason);
+
+	// clear the death timer
+	GetWorld()->GetTimerManager().ClearTimer(DeathTimer);
+}
diff --git a/Templates/TP_ThirdPerson/Source/TP_ThirdPerson/Variant_Combat/AI/CombatEnemy.h b/Templates/TP_ThirdPerson/Source/TP_ThirdPerson/Variant_Combat/AI/CombatEnemy.h
new file mode 100644
--- /dev/null
+++ b/Templates/TP_ThirdPerson/Source/TP_ThirdPerson/Variant_Combat/AI/CombatEnemy.h
@@ -0,0 +1,217 @@
+// Copyright Epic Games, Inc. All Rights Reserved.
+
+#pragma once
+
+#include "CoreMinimal.h"
+#include "GameFramework/Character.h"
+#include "CombatAttacker.h"
+#include "CombatDamageable.h"
+#include "Animation/AnimMontage.h"
+#include "Engine/TimerHandle.h"
+#include "CombatEnemy.generated.h"
+
+class UWidgetComponent;
+class UCombatLifeBar;
+class UAnimMontage;
+
+/** Completed attack animation delegate for StateTree */
+DECLARE_DELEGATE(FOnEnemyAttackCompleted);
+
+/** Landed delegate for StateTree */
+DECLARE_DELEGATE(FOnEnemyLanded);
+
+/** Enemy died delegate */
+DECLARE_DYNAMIC_MULTICAST_DELEGATE(FOnEnemyDied);
+
+/**
+ *  An AI-controlled character with combat capabilities.
+ *  Its bundled AI Controller runs logic through StateTree
+ */
+UCLASS(abstract)
+class ACombatEnemy : public ACharacter, public ICombatAttacker, public ICombatDamageable
+{
+	GENERATED_BODY()
+
+	/** Life bar widget component */
+	UPROPERTY(VisibleAnywhere, BlueprintReadOnly, Category = UI, meta = (AllowPrivateAccess = "true"))
+	UWidgetComponent* LifeBar;
+
+public:
+	
+	/** Constructor */
+	ACombatEnemy();
+
+protected:
+
+	/** Max amount of HP the character will have on respawn */
+	UPROPERTY(EditAnywhere, Category="Damage")
+	float MaxHP = 3.0f;
+
+public:
+
+	/** Current amount of HP the character has */
+	UPROPERTY(VisibleAnywhere, BlueprintReadOnly, Category="Damage", meta = (ClampMin = 0, ClampMax = 100))
+	float CurrentHP = 0.0f;
+
+protected:
+
+	/** Name of the pelvis bone, for damage ragdoll physics */
+	UPROPERTY(EditAnywhere, Category="Damage")
+	FName PelvisBoneName;
+
+	/** Pointer to the life bar widget */
+	UPROPERTY(EditAnywhere, Category="Damage")
+	UCombatLifeBar* LifeBarWidget;
+
+	/** If true, the character is currently playing an attack animation */
+	bool bIsAttacking = false;
+
+	/** Distance ahead of the character that melee attack sphere collision traces will extend */
+	UPROPERTY(EditAnywhere, Category="Melee Attack|Trace", meta = (ClampMin = 0, ClampMax = 500, Units = "cm"))
+	float MeleeTraceDistance = 75.0f;
+
+	/** Radius of the sphere trace for melee attacks */
+	UPROPERTY(EditAnywhere, Category="Melee Attack|Trace", meta = (ClampMin = 0, ClampMax = 500, Units = "cm"))
+	float MeleeTraceRadius = 50.0f;
+
+	/** Amount of damage a melee attack will deal */
+	UPROPERTY(EditAnywhere, Category="Melee Attack|Damage", meta = (ClampMin = 0, ClampMax = 100))
+	float MeleeDamage = 1.0f;
+
+	/** Amount of knockback impulse a melee attack will apply */
+	UPROPERTY(EditAnywhere, Category="Melee Attack|Damage", meta = (ClampMin = 0, ClampMax = 1000, Units = "cm/s"))
+	float MeleeKnockbackImpulse = 150.0f;
+
+	/** Amount of upwards impulse a melee attack will apply */
+	UPROPERTY(EditAnywhere, Category="Melee Attack|Damage", meta = (ClampMin = 0, ClampMax = 1000, Units = "cm/s"))
+	float MeleeLaunchImpulse = 350.0f;
+
+	/** AnimMontage that will play for combo attacks */
+	UPROPERTY(EditAnywhere, Category="Melee Attack|Combo")
+	UAnimMontage* ComboAttackMontage;
+
+	/** Names of the AnimMontage sections that correspond to each stage of the combo attack */
+	UPROPERTY(EditAnywhere, Category="Melee Attack|Combo")
+	TArray<FName> ComboSectionNames;
+
+	/** Target number of attacks in the combo attack string we're playing */
+	int32 TargetComboCount = 0;
+
+	/** Index of the current stage of the melee attack combo */
+	int32 CurrentComboAttack = 0;
+
+	/** AnimMontage that will play for charged attacks */
+	UPROPERTY(EditAnywhere, Category="Melee Attack|Charged")
+	UAnimMontage* ChargedAttackMontage;
+
+	/** Name of the AnimMontage section that corresponds to the charge loop */
+	UPROPERTY(EditAnywhere, Category="Melee Attack|Charged")
+	FName ChargeLoopSection;
+
+	/** Name of the AnimMontage section that corresponds to the attack */
+	UPROPERTY(EditAnywhere, Category="Melee Attack|Charged")
+	FName ChargeAttackSection;
+
+	/** Minimum number of charge animation loops that will be played by the AI */
+	UPROPERTY(EditAnywhere, Category="Melee Attack|Charged", meta = (ClampMin = 1, ClampMax = 20))
+	int32 MinChargeLoops = 2;
+
+	/** Maximum number of charge animation loops that will be played by the AI */
+	UPROPERTY(EditAnywhere, Category="Melee Attack|Charged", meta = (ClampMin = 1, ClampMax = 20))
+	int32 MaxChargeLoops = 5;
+
+	/** Target number of charge animation loops to play in this charged attack */
+	int32 TargetChargeLoops = 0;
+
+	/** Number of charge animation loop currently playing */
+	int32 CurrentChargeLoop = 0;
+
+	/** Time to wait before removing this character from the level after it dies */
+	UPROPERTY(EditAnywhere, Category="Death")
+	float DeathRemovalTime = 5.0f;
+
+	/** Enemy death timer */
+	FTimerHandle DeathTimer;
+
+	/** Attack montage ended delegate */
+	FOnMontageEnded OnAttackMontageEnded;
+
+public:
+	/** Attack completed internal delegate to notify StateTree tasks */
+	FOnEnemyAttackCompleted OnAttackCompleted;
+
+	/** Landed internal delegate to notify StateTree tasks. We use this instead of the built-in Landed delegate so we can bind to a Lambda in StateTree tasks */
+	FOnEnemyLanded OnEnemyLanded;
+
+	/** Enemy died delegate. Allows external subscribers to respond to enemy death */
+	UPROPERTY(BlueprintAssignable, Category="Events")
+	FOnEnemyDied OnEnemyDied;
+
+public:
+
+	/** Performs an AI-initiated combo attack. Number of hits will be decided by this character */
+	void DoAIComboAttack();
+
+	/** Performs an AI-initiated charged attack. Charge time will be decided by this character */
+	void DoAIChargedAttack();
+
+	/** Called from a delegate when the attack montage ends */
+	void AttackMontageEnded(UAnimMontage* Montage, bool bInterrupted);
+
+public:
+
+	// ~begin ICombatAttacker interface
+
+	/** Performs an attack's collision check */
+	virtual void DoAttackTrace(FName DamageSourceBone) override;
+
+	/** Performs a combo attack's check to continue the string */
+	UFUNCTION(BlueprintCallable, Category="Attacker")
+	virtual void CheckCombo() override;
+
+	/** Performs a charged attack's check to loop the charge animation */
+	UFUNCTION(BlueprintCallable, Category="Attacker")
+	virtual void CheckChargedAttack() override;
+
+	// ~end ICombatAttacker interface
+
+	// ~begin ICombatDamageable interface
+
+	/** Handles damage and knockback events */
+	virtual void ApplyDamage(float Damage, AActor* DamageCauser, const FVector& DamageLocation, const FVector& DamageImpulse) override;
+
+	/** Handles death events */
+	virtual void HandleDeath() override;
+
+	/** Handles healing events */
+	virtual void ApplyHealing(float Healing, AActor* Healer) override;
+
+	// ~end ICombatDamageable interface
+
+protected:
+
+	/** Removes this character from the level after it dies */
+	void RemoveFromLevel();
+
+public:
+
+	/** Overrides the default TakeDamage functionality */
+	virtual float TakeDamage(float Damage, struct FDamageEvent const& DamageEvent, AController* EventInstigator, AActor* DamageCauser) override;
+
+	/** Overrides landing to reset damage ragdoll physics */
+	virtual void Landed(const FHitResult& Hit) override;
+
+protected:
+
+	/** Blueprint handler to play damage received effects */
+	UFUNCTION(BlueprintImplementableEvent, Category="Combat")
+	void ReceivedDamage(float Damage, const FVector& ImpactPoint, const FVector& DamageDirection);
+
+protected:
+
+	/** Gameplay initialization */
+	virtual void BeginPlay() override;
+
+	/** EndPlay cleanup */
+	virtual void EndPlay(EEndPlayReason::Type EndPlayReason) override;
+};
diff --git a/Templates/TP_ThirdPerson/Source/TP_ThirdPerson/Variant_Combat/AI/CombatEnemySpawner.cpp b/Templates/TP_ThirdPerson/Source/TP_ThirdPerson/Variant_Combat/AI/CombatEnemySpawner.cpp
new file mode 100644
--- /dev/null
+++ b/Templates/TP_ThirdPerson/Source/TP_ThirdPerson/Variant_Combat/AI/CombatEnemySpawner.cpp
@@ -0,0 +1,125 @@
+// Copyright Epic Games, Inc. All Rights Reserved.
+
+
+#include "CombatEnemySpawner.h"
+#include "Components/SceneComponent.h"
+#include "Components/CapsuleComponent.h"
+#include "Components/ArrowComponent.h"
+#include "TimerManager.h"
+#include "CombatEnemy.h"
+
+ACombatEnemySpawner::ACombatEnemySpawner()
+{
+	PrimaryActorTick.bCanEverTick = false;
+
+	// create the root
+	RootComponent = CreateDefaultSubobject<USceneComponent>(TEXT("Root"));
+
+	// create the reference spawn capsule
+	SpawnCapsule = CreateDefaultSubobject<UCapsuleComponent>(TEXT("Spawn Capsule"));
+	SpawnCapsule->SetupAttachment(RootComponent);
+
+	SpawnCapsule->SetRelativeLocation(FVector(0.0f, 0.0f, 90.0f));
+	SpawnCapsule->SetCapsuleSize(35.0f, 90.0f);
+	SpawnCapsule->SetCollisionProfileName(FName("NoCollision"));
+
+	SpawnDirection = CreateDefaultSubobject<UArrowComponent>(TEXT("Spawn Direction"));
+	SpawnDirection->SetupAttachment(RootComponent);
+}
+
+void ACombatEnemySpawner::BeginPlay()
+{
+	Super::BeginPlay();
+	
+	// should we spawn an enemy right away?
+	if (bShouldSpawnEnemiesImmediately)
+	{
+		// schedule the first enemy spawn
+		GetWorld()->GetTimerManager().SetTimer(SpawnTimer, this, &ACombatEnemySpawner::SpawnEnemy, InitialSpawnDelay);
+	}
+
+}
+
+void ACombatEnemySpawner::EndPlay(EEndPlayReason::Type EndPlayReason)
+{
+	Super::EndPlay(EndPlayReason);
+
+	// clear the spawn timer
+	GetWorld()->GetTimerManager().ClearTimer(SpawnTimer);
+}
+
+void ACombatEnemySpawner::SpawnEnemy()
+{
+	// ensure the enemy class is valid
+	if (IsValid(EnemyClass))
+	{
+		// spawn the enemy at the reference capsule's transform
+		FActorSpawnParameters SpawnParams;
+		SpawnParams.SpawnCollisionHandlingOverride = ESpawnActorCollisionHandlingMethod::AdjustIfPossibleButAlwaysSpawn;
+
+		ACombatEnemy* SpawnedEnemy = GetWorld()->SpawnActor<ACombatEnemy>(EnemyClass, SpawnCapsule->GetComponentTransform(), SpawnParams);
+
+		// was the enemy successfully created?
+		if (SpawnedEnemy)
+		{
+			// subscribe to the death delegate
+			SpawnedEnemy->OnEnemyDied.AddDynamic(this, &ACombatEnemySpawner::OnEnemyDied);
+		}
+	}
+}
+
+void ACombatEnemySpawner::OnEnemyDied()
+{
+	// decrease the spawn counter
+	--SpawnCount;
+
+	// is this the last enemy we should spawn?
+	if (SpawnCount <= 0)
+	{
+		// schedule the activation on depleted message
+		GetWorld()->GetTimerManager().SetTimer(SpawnTimer, this, &ACombatEnemySpawner::SpawnerDepleted, ActivationDelay);
+		return;
+	}
+
+	// schedule the next enemy spawn
+	GetWorld()->GetTimerManager().SetTimer(SpawnTimer, this, &ACombatEnemySpawner::SpawnEnemy, RespawnDelay);
+}
+
+void ACombatEnemySpawner::SpawnerDepleted()
+{
+	// process the actors to activate list
+	for (AActor* CurrentActor : ActorsToActivateWhenDepleted)
+	{
+		// check if the actor is activatable
+		if (ICombatActivatable* CombatActivatable = Cast<ICombatActivatable>(CurrentActor))
+		{
+			// activate the actor
+			CombatActivatable->ActivateInteraction(this);
+		}
+	}
+}
+
+void ACombatEnemySpawner::ToggleInteraction(AActor* ActivationInstigator)
+{
+	// stub
+}
+
+void ACombatEnemySpawner::ActivateInteraction(AActor* ActivationInstigator)
+{
+	// ensure we're only activated once, and only if we've deferred enemy spawning
+	if (bHasBeenActivated || bShouldSpawnEnemiesImmediately)
+	{
+		return;
+	}
+
+	// raise the activation flag
+	bHasBeenActivated = true;
+
+	// spawn the first enemy
+	SpawnEnemy();
+}
+
+void ACombatEnemySpawner::DeactivateInteraction(AActor* ActivationInstigator)
+{
+	// stub
+}
diff --git a/Templates/TP_ThirdPerson/Source/TP_ThirdPerson/Variant_Combat/AI/CombatEnemySpawner.h b/Templates/TP_ThirdPerson/Source/TP_ThirdPerson/Variant_Combat/AI/CombatEnemySpawner.h
new file mode 100644
--- /dev/null
+++ b/Templates/TP_ThirdPerson/Source/TP_ThirdPerson/Variant_Combat/AI/CombatEnemySpawner.h
@@ -0,0 +1,109 @@
+// Copyright Epic Games, Inc. All Rights Reserved.
+
+#pragma once
+
+#include "CoreMinimal.h"
+#include "GameFramework/Actor.h"
+#include "CombatActivatable.h"
+#include "CombatEnemySpawner.generated.h"
+
+class UCapsuleComponent;
+class UArrowComponent;
+class ACombatEnemy;
+
+/**
+ *  A basic Actor in charge of spawning Enemy Characters and monitoring their deaths.
+ *  Enemies will be spawned one by one, and the spawner will wait until the enemy dies before spawning a new one.
+ *  The spawner can be remotely activated through the ICombatActivatable interface
+ *  When the last spawned enemy dies, the spawner can also activate other ICombatActivatables
+ */
+UCLASS(abstract)
+class ACombatEnemySpawner : public AActor, public ICombatActivatable
+{
+	GENERATED_BODY()
+	
+	UPROPERTY(VisibleAnywhere, BlueprintReadOnly, Category = Components, meta = (AllowPrivateAccess = "true"))
+	UCapsuleComponent* SpawnCapsule;
+
+	UPROPERTY(VisibleAnywhere, BlueprintReadOnly, Category = Components, meta = (AllowPrivateAccess = "true"))
+	UArrowComponent* SpawnDirection;
+
+protected:
+
+	/** Type of enemy to spawn */
+	UPROPERTY(EditAnywhere, BlueprintReadOnly, Category="Enemy Spawner")
+	TSubclassOf<ACombatEnemy> EnemyClass;
+
+	/** If true, the first enemy will be spawned as soon as the game starts */
+	UPROPERTY(EditAnywhere, BlueprintReadOnly, Category="Enemy Spawner")
+	bool bShouldSpawnEnemiesImmediately = true;
+
+	/** Time to wait before spawning the first enemy on game start */
+	UPROPERTY(EditAnywhere, BlueprintReadOnly, Category="Enemy Spawner", meta = (ClampMin = 0, ClampMax = 10))
+	float InitialSpawnDelay = 5.0f;
+
+	/** Number of enemies to spawn */
+	UPROPERTY(EditAnywhere, BlueprintReadOnly, Category="Enemy Spawner", meta = (ClampMin = 0, ClampMax = 100))
+	int32 SpawnCount = 1;
+
+	/** Time to wait before spawning the next enemy after the current one dies */
+	UPROPERTY(EditAnywhere, BlueprintReadOnly, Category="Enemy Spawner", meta = (ClampMin = 0, ClampMax = 10))
+	float RespawnDelay = 5.0f;
+
+	/** Time to wait after this spawner is depleted before activating the actor list */
+	UPROPERTY(EditAnywhere, BlueprintReadOnly, Category="Activation", meta = (ClampMin = 0, ClampMax = 10))
+	float ActivationDelay = 1.0f;
+
+	/** List of actors to activate after the last enemy dies */
+	UPROPERTY(EditAnywhere, BlueprintReadOnly, Category="Activation")
+	TArray<AActor*> ActorsToActivateWhenDepleted;
+
+	/** Flag to ensure this is only activated once */
+	bool bHasBeenActivated = false;
+
+	/** Timer to spawn enemies after a delay */
+	FTimerHandle SpawnTimer;
+
+public:	
+	
+	/** Constructor */
+	ACombatEnemySpawner();
+
+public:
+
+	/** Initialization */
+	virtual void BeginPlay() override;
+
+	/** Cleanup */
+	virtual void EndPlay(EEndPlayReason::Type EndPlayReason) override;
+
+protected:
+
+	/** Spawn an enemy and subscribe to its death event */
+	void SpawnEnemy();
+
+	/** Called when the spawned enemy has died */
+	UFUNCTION()
+	void OnEnemyDied();
+
+	/** Called after the last spawned enemy has died */
+	void SpawnerDepleted();
+
+public:
+
+	// ~begin ICombatActivatable interface
+
+	/** Toggles the Spawner */
+	UFUNCTION(BlueprintCallable, Category="Activatable")
+	virtual void ToggleInteraction(AActor* ActivationInstigator) override;
+
+	/** Activates the Spawner */
+	UFUNCTION(BlueprintCallable, Category="Activatable")
+	virtual void ActivateInteraction(AActor* ActivationInstigator) override;
+
+	/** Deactivates the Spawner */
+	UFUNCTION(BlueprintCallable, Category="Activatable")
+	virtual void DeactivateInteraction(AActor* ActivationInstigator) override;
+
+	// ~end IActivatable interface
+};
diff --git a/Templates/TP_ThirdPerson/Source/TP_ThirdPerson/Variant_Combat/AI/CombatStateTreeUtility.cpp b/Templates/TP_ThirdPerson/Source/TP_ThirdPerson/Variant_Combat/AI/CombatStateTreeUtility.cpp
new file mode 100644
--- /dev/null
+++ b/Templates/TP_ThirdPerson/Source/TP_ThirdPerson/Variant_Combat/AI/CombatStateTreeUtility.cpp
@@ -0,0 +1,292 @@
+// Copyright Epic Games, Inc. All Rights Reserved.
+
+
+#include "CombatStateTreeUtility.h"
+#include "StateTreeExecutionContext.h"
+#include "StateTreeExecutionTypes.h"
+#include "GameFramework/Character.h"
+#include "GameFramework/CharacterMovementComponent.h"
+#include "AIController.h"
+#include "CombatEnemy.h"
+#include "Kismet/GameplayStatics.h"
+#include "StateTreeAsyncExecutionContext.h"
+
+#define LOCTEXT_NAMESPACE "ThirdPersonTemplate"
+
+bool FStateTreeCharacterGroundedCondition::TestCondition(FStateTreeExecutionContext& Context) const
+{
+	const FInstanceDataType& InstanceData = Context.GetInstanceData(*this);
+
+	// is the character currently grounded?
+	bool bCondition = InstanceData.Character->GetMovementComponent()->IsMovingOnGround();
+
+	return InstanceData.bMustBeOnAir ? !bCondition : bCondition;
+}
+
+#if WITH_EDITOR
+FText FStateTreeCharacterGroundedCondition::GetDescription(const FGuid& ID, FStateTreeDataView InstanceDataView, const IStateTreeBindingLookup& BindingLookup, EStateTreeNodeFormatting Formatting /*= EStateTreeNodeFormatting::Text*/) const
+{
+	return LOCTEXT("StateTreeConditionGrounded", "<b>Is Character Grounded</b>");
+}
+#endif // WITH_EDITOR
+
+////////////////////////////////////////////////////////////////////
+
+EStateTreeRunStatus FStateTreeComboAttackTask::EnterState(FStateTreeExecutionContext& Context, const FStateTreeTransitionResult& Transition) const
+{
+	// have we transitioned from another state?
+	if (Transition.ChangeType == EStateTreeStateChangeType::Changed)
+	{
+		// get the instance data
+		FInstanceDataType& InstanceData = Context.GetInstanceData(*this);
+
+		// bind to the on attack completed delegate
+		InstanceData.Character->OnAttackCompleted.BindLambda(
+			[WeakContext = Context.MakeWeakExecutionContext()]()
+			{
+				WeakContext.FinishTask(EStateTreeFinishTaskType::Succeeded);
+			}
+		);
+
+
+		// tell the character to do a combo attack
+		InstanceData.Character->DoAIComboAttack();
+	}
+
+	return EStateTreeRunStatus::Running;
+}
+
+void FStateTreeComboAttackTask::ExitState(FStateTreeExecutionContext& Context, const FStateTreeTransitionResult& Transition) const
+{
+	// have we transitioned from another state?
+	if (Transition.ChangeType == EStateTreeStateChangeType::Changed)
+	{
+		// get the instance data
+		FInstanceDataType& InstanceData = Context.GetInstanceData(*this);
+
+		// unbind the on attack completed delegate
+		InstanceData.Character->OnAttackCompleted.Unbind();
+	}
+}
+
+#if WITH_EDITOR
+FText FStateTreeComboAttackTask::GetDescription(const FGuid& ID, FStateTreeDataView InstanceDataView, const IStateTreeBindingLookup& BindingLookup, EStateTreeNodeFormatting Formatting /*= EStateTreeNodeFormatting::Text*/) const
+{
+	return LOCTEXT("StateTreeTaskComboAttackDescription", "<b>Do Combo Attack</b>");
+}
+#endif // WITH_EDITOR
+
+////////////////////////////////////////////////////////////////////
+
+EStateTreeRunStatus FStateTreeChargedAttackTask::EnterState(FStateTreeExecutionContext& Context, const FStateTreeTransitionResult& Transition) const
+{
+	// have we transitioned from another state?
+	if (Transition.ChangeType == EStateTreeStateChangeType::Changed)
+	{
+		// get the instance data
+		FInstanceDataType& InstanceData = Context.GetInstanceData(*this);
+
+		// bind to the on attack completed delegate
+		InstanceData.Character->OnAttackCompleted.BindLambda(
+			[WeakContext = Context.MakeWeakExecutionContext()]()
+			{
+				WeakContext.FinishTask(EStateTreeFinishTaskType::Succeeded);
+			}
+		);
+
+		// tell the character to do a combo attack
+		InstanceData.Character->DoAIChargedAttack();
+	}
+
+	return EStateTreeRunStatus::Running;
+}
+
+void FStateTreeChargedAttackTask::ExitState(FStateTreeExecutionContext& Context, const FStateTreeTransitionResult& Transition) const
+{
+	// have we transitioned from another state?
+	if (Transition.ChangeType == EStateTreeStateChangeType::Changed)
+	{
+		// get the instance data
+		FInstanceDataType& InstanceData = Context.GetInstanceData(*this);
+
+		// unbind the on attack completed delegate
+		InstanceData.Character->OnAttackCompleted.Unbind();
+	}
+}
+
+#if WITH_EDITOR
+FText FStateTreeChargedAttackTask::GetDescription(const FGuid& ID, FStateTreeDataView InstanceDataView, const IStateTreeBindingLookup& BindingLookup, EStateTreeNodeFormatting Formatting /*= EStateTreeNodeFormatting::Text*/) const
+{
+	return LOCTEXT("StateTreeTaskChargedAttackDescription", "<b>Do Charged Attack</b>");
+}
+#endif // WITH_EDITOR
+
+////////////////////////////////////////////////////////////////////
+
+EStateTreeRunStatus FStateTreeWaitForLandingTask::EnterState(FStateTreeExecutionContext& Context, const FStateTreeTransitionResult& Transition) const
+{
+	// have we transitioned from another state?
+	if (Transition.ChangeType == EStateTreeStateChangeType::Changed)
+	{
+		// get the instance data
+		FInstanceDataType& InstanceData = Context.GetInstanceData(*this);
+
+		// bind to the on enemy landed delegate
+		InstanceData.Character->OnEnemyLanded.BindLambda(
+			[WeakContext = Context.MakeWeakExecutionContext()]()
+			{
+				WeakContext.FinishTask(EStateTreeFinishTaskType::Succeeded);
+			}
+		);
+	}
+
+	return EStateTreeRunStatus::Running;
+}
+
+void FStateTreeWaitForLandingTask::ExitState(FStateTreeExecutionContext& Context, const FStateTreeTransitionResult& Transition) const
+{
+	// have we transitioned from another state?
+	if (Transition.ChangeType == EStateTreeStateChangeType::Changed)
+	{
+		// get the instance data
+		FInstanceDataType& InstanceData = Context.GetInstanceData(*this);
+
+		// bind the on enemy landed delegate
+		InstanceData.Character->OnEnemyLanded.Unbind();
+	}
+}
+
+#if WITH_EDITOR
+FText FStateTreeWaitForLandingTask::GetDescription(const FGuid& ID, FStateTreeDataView InstanceDataView, const IStateTreeBindingLookup& BindingLookup, EStateTreeNodeFormatting Formatting /*= EStateTreeNodeFormatting::Text*/) const
+{
+	return LOCTEXT("StateTreeTaskChargedAttackDescription", "<b>Wait for Landing</b>");
+}
+#endif // WITH_EDITOR
+
+////////////////////////////////////////////////////////////////////
+
+EStateTreeRunStatus FStateTreeFaceActorTask::EnterState(FStateTreeExecutionContext& Context, const FStateTreeTransitionResult& Transition) const
+{
+	// have we transitioned from another state?
+	if (Transition.ChangeType == EStateTreeStateChangeType::Changed)
+	{
+		// get the instance data
+		FInstanceDataType& InstanceData = Context.GetInstanceData(*this);
+
+		// set the AI Controller's focus
+		InstanceData.Controller->SetFocus(InstanceData.ActorToFaceTowards);
+	}
+
+	return EStateTreeRunStatus::Running;
+}
+
+void FStateTreeFaceActorTask::ExitState(FStateTreeExecutionContext& Context, const FStateTreeTransitionResult& Transition) const
+{
+	// have we transitioned to another state?
+	if (Transition.ChangeType == EStateTreeStateChangeType::Changed)
+	{
+		// get the instance data
+		FInstanceDataType& InstanceData = Context.GetInstanceData(*this);
+
+		// clear the AI Controller's focus
+		InstanceData.Controller->ClearFocus(EAIFocusPriority::Gameplay);
+	}
+}
+
+#if WITH_EDITOR
+FText FStateTreeFaceActorTask::GetDescription(const FGuid& ID, FStateTreeDataView InstanceDataView, const IStateTreeBindingLookup& BindingLookup, EStateTreeNodeFormatting Formatting /*= EStateTreeNodeFormatting::Text*/) const
+{
+	return LOCTEXT("StateTreeTaskFaceActorDescription", "<b>Face Towards Actor</b>");
+}
+#endif // WITH_EDITOR
+
+////////////////////////////////////////////////////////////////////
+
+EStateTreeRunStatus FStateTreeFaceLocationTask::EnterState(FStateTreeExecutionContext& Context, const FStateTreeTransitionResult& Transition) const
+{
+	// have we transitioned from another state?
+	if (Transition.ChangeType == EStateTreeStateChangeType::Changed)
+	{
+		// get the instance data
+		FInstanceDataType& InstanceData = Context.GetInstanceData(*this);
+
+		// set the AI Controller's focus
+		InstanceData.Controller->SetFocalPoint(InstanceData.FaceLocation);
+	}
+
+	return EStateTreeRunStatus::Running;
+}
+
+void FStateTreeFaceLocationTask::ExitState(FStateTreeExecutionContext& Context, const FStateTreeTransitionResult& Transition) const
+{
+	// have we transitioned to another state?
+	if (Transition.ChangeType == EStateTreeStateChangeType::Changed)
+	{
+		// get the instance data
+		FInstanceDataType& InstanceData = Context.GetInstanceData(*this);
+
+		// clear the AI Controller's focus
+		InstanceData.Controller->ClearFocus(EAIFocusPriority::Gameplay);
+	}
+}
+
+#if WITH_EDITOR
+FText FStateTreeFaceLocationTask::GetDescription(const FGuid& ID, FStateTreeDataView InstanceDataView, const IStateTreeBindingLookup& BindingLookup, EStateTreeNodeFormatting Formatting /*= EStateTreeNodeFormatting::Text*/) const
+{
+	return LOCTEXT("StateTreeTaskFaceLocationDescription", "<b>Face Towards Location</b>");
+}
+#endif // WITH_EDITOR
+
+////////////////////////////////////////////////////////////////////
+
+EStateTreeRunStatus FStateTreeSetCharacterSpeedTask::EnterState(FStateTreeExecutionContext& Context, const FStateTreeTransitionResult& Transition) const
+{
+	// have we transitioned from another state?
+	if (Transition.ChangeType == EStateTreeStateChangeType::Changed)
+	{
+		// get the instance data
+		FInstanceDataType& InstanceData = Context.GetInstanceData(*this);
+
+		// set the character's max ground speed
+		InstanceData.Character->GetCharacterMovement()->MaxWalkSpeed = InstanceData.Speed;
+	}
+
+	return EStateTreeRunStatus::Running;
+}
+
+#if WITH_EDITOR
+FText FStateTreeSetCharacterSpeedTask::GetDescription(const FGuid& ID, FStateTreeDataView InstanceDataView, const IStateTreeBindingLookup& BindingLookup, EStateTreeNodeFormatting Formatting /*= EStateTreeNodeFormatting::Text*/) const
+{
+	return LOCTEXT("StateTreeTaskSetCharacterSpeedDescription", "<b>Set Character Speed</b>");
+}
+#endif // WITH_EDITOR
+
+////////////////////////////////////////////////////////////////////
+
+EStateTreeRunStatus FStateTreeGetPlayerInfoTask::Tick(FStateTreeExecutionContext& Context, const float DeltaTime) const
+{
+	// get the instance data
+	FInstanceDataType& InstanceData = Context.GetInstanceData(*this);
+
+	// get the character possessed by the first local player
+	InstanceData.TargetPlayerCharacter = Cast<ACharacter>(UGameplayStatics::GetPlayerPawn(InstanceData.Character, 0));
+
+	// do we have a valid target?
+	if (InstanceData.TargetPlayerCharacter)
+	{
+		// update the last known location
+		InstanceData.TargetPlayerLocation = InstanceData.TargetPlayerCharacter->GetActorLocation();
+	}
+
+	// update the distance
+	InstanceData.DistanceToTarget = FVector::Distance(InstanceData.TargetPlayerLocation, InstanceData.Character->GetActorLocation());
+
+	return EStateTreeRunStatus::Running;
+}
+
+#if WITH_EDITOR
+FText FStateTreeGetPlayerInfoTask::GetDescription(const FGuid& ID, FStateTreeDataView InstanceDataView, const IStateTreeBindingLookup& BindingLookup, EStateTreeNodeFormatting Formatting /*= EStateTreeNodeFormatting::Text*/) const
+{
+	return LOCTEXT("StateTreeTaskGetPlayerInfoDescription", "<b>Get Player Info</b>");
+}
+#endif // WITH_EDITOR
\ No newline at end of file
diff --git a/Templates/TP_ThirdPerson/Source/TP_ThirdPerson/Variant_Combat/AI/CombatStateTreeUtility.h b/Templates/TP_ThirdPerson/Source/TP_ThirdPerson/Variant_Combat/AI/CombatStateTreeUtility.h
new file mode 100644
--- /dev/null
+++ b/Templates/TP_ThirdPerson/Source/TP_ThirdPerson/Variant_Combat/AI/CombatStateTreeUtility.h
@@ -0,0 +1,311 @@
+// Copyright Epic Games, Inc. All Rights Reserved.
+
+#pragma once
+
+#include "CoreMinimal.h"
+#include "StateTreeTaskBase.h"
+#include "StateTreeConditionBase.h"
+
+#include "CombatStateTreeUtility.generated.h"
+
+class ACharacter;
+class AAIController;
+class ACombatEnemy;
+
+/**
+ *  Instance data struct for the FStateTreeCharacterGroundedCondition condition
+ */
+USTRUCT()
+struct FStateTreeCharacterGroundedConditionInstanceData
+{
+	GENERATED_BODY()
+	
+	/** Character to check grounded status on */
+	UPROPERTY(EditAnywhere, Category = "Context")
+	ACharacter* Character;
+
+	/** If true, the condition passes if the character is not grounded instead */
+	UPROPERTY(EditAnywhere, Category = "Condition")
+	bool bMustBeOnAir = false;
+};
+STATETREE_POD_INSTANCEDATA(FStateTreeCharacterGroundedConditionInstanceData);
+
+/**
+ *  StateTree condition to check if the character is grounded
+ */
+USTRUCT(DisplayName = "Character is Grounded")
+struct FStateTreeCharacterGroundedCondition : public FStateTreeConditionCommonBase
+{
+	GENERATED_BODY()
+
+	/** Set the instance data type */
+	using FInstanceDataType = FStateTreeCharacterGroundedConditionInstanceData;
+	virtual const UStruct* GetInstanceDataType() const override { return FInstanceDataType::StaticStruct(); }
+
+	/** Default constructor */
+	FStateTreeCharacterGroundedCondition() = default;
+	
+	/** Tests the StateTree condition */
+	virtual bool TestCondition(FStateTreeExecutionContext& Context) const override;
+
+#if WITH_EDITOR
+
+	/** Provides the description string */
+	virtual FText GetDescription(const FGuid& ID, FStateTreeDataView InstanceDataView, const IStateTreeBindingLookup& BindingLookup, EStateTreeNodeFormatting Formatting = EStateTreeNodeFormatting::Text) const override;
+#endif
+
+};
+
+////////////////////////////////////////////////////////////////////
+
+/**
+ *  Instance data struct for the Combat StateTree tasks
+ */
+USTRUCT()
+struct FStateTreeAttackInstanceData
+{
+	GENERATED_BODY()
+
+	/** Character that will perform the attack */
+	UPROPERTY(EditAnywhere, Category = Context)
+	TObjectPtr<ACombatEnemy> Character;
+};
+
+/**
+ *  StateTree task to perform a combo attack
+ */
+USTRUCT(meta=(DisplayName="Combo Attack", Category="Combat"))
+struct FStateTreeComboAttackTask : public FStateTreeTaskCommonBase
+{
+	GENERATED_BODY()
+
+	/* Ensure we're using the correct instance data struct */
+	using FInstanceDataType = FStateTreeAttackInstanceData;
+	virtual const UStruct* GetInstanceDataType() const override { return FInstanceDataType::StaticStruct(); }
+
+	/** Runs when the owning state is entered */
+	virtual EStateTreeRunStatus EnterState(FStateTreeExecutionContext& Context, const FStateTreeTransitionResult& Transition) const override;
+
+	/** Runs when the owning state is ended */
+	virtual void ExitState(FStateTreeExecutionContext& Context, const FStateTreeTransitionResult& Transition) const override;
+
+#if WITH_EDITOR
+	virtual FText GetDescription(const FGuid& ID, FStateTreeDataView InstanceDataView, const IStateTreeBindingLookup& BindingLookup, EStateTreeNodeFormatting Formatting = EStateTreeNodeFormatting::Text) const override;
+#endif // WITH_EDITOR
+};
+
+/**
+ *  StateTree task to perform a charged attack
+ */
+USTRUCT(meta=(DisplayName="Charged Attack", Category="Combat"))
+struct FStateTreeChargedAttackTask : public FStateTreeTaskCommonBase
+{
+	GENERATED_BODY()
+
+	/* Ensure we're using the correct instance data struct */
+	using FInstanceDataType = FStateTreeAttackInstanceData;
+	virtual const UStruct* GetInstanceDataType() const override { return FInstanceDataType::StaticStruct(); }
+
+	/** Runs when the owning state is entered */
+	virtual EStateTreeRunStatus EnterState(FStateTreeExecutionContext& Context, const FStateTreeTransitionResult& Transition) const override;
+
+	/** Runs when the owning state is ended */
+	virtual void ExitState(FStateTreeExecutionContext& Context, const FStateTreeTransitionResult& Transition) const override;
+
+#if WITH_EDITOR
+	virtual FText GetDescription(const FGuid& ID, FStateTreeDataView InstanceDataView, const IStateTreeBindingLookup& BindingLookup, EStateTreeNodeFormatting Formatting = EStateTreeNodeFormatting::Text) const override;
+#endif // WITH_EDITOR
+};
+
+/**
+ *  StateTree task to wait for the character to land
+ */
+USTRUCT(meta=(DisplayName="Wait for Landing", Category="Combat"))
+struct FStateTreeWaitForLandingTask : public FStateTreeTaskCommonBase
+{
+	GENERATED_BODY()
+
+	/* Ensure we're using the correct instance data struct */
+	using FInstanceDataType = FStateTreeAttackInstanceData;
+	virtual const UStruct* GetInstanceDataType() const override { return FInstanceDataType::StaticStruct(); }
+
+	/** Runs when the owning state is entered */
+	virtual EStateTreeRunStatus EnterState(FStateTreeExecutionContext& Context, const FStateTreeTransitionResult& Transition) const override;
+
+	/** Runs when the owning state is ended */
+	virtual void ExitState(FStateTreeExecutionContext& Context, const FStateTreeTransitionResult& Transition) const override;
+
+#if WITH_EDITOR
+	virtual FText GetDescription(const FGuid& ID, FStateTreeDataView InstanceDataView, const IStateTreeBindingLookup& BindingLookup, EStateTreeNodeFormatting Formatting = EStateTreeNodeFormatting::Text) const override;
+#endif // WITH_EDITOR
+};
+
+////////////////////////////////////////////////////////////////////
+
+/**
+ *  Instance data struct for the Face Towards Actor StateTree task
+ */
+USTRUCT()
+struct FStateTreeFaceActorInstanceData
+{
+	GENERATED_BODY()
+
+	/** AI Controller that will determine the focused actor */
+	UPROPERTY(EditAnywhere, Category = Context)
+	TObjectPtr<AAIController> Controller;
+
+	/** Actor that will be faced towards */
+	UPROPERTY(EditAnywhere, Category = Input)
+	TObjectPtr<AActor> ActorToFaceTowards;
+};
+
+/**
+ *  StateTree task to face an AI-Controlled Pawn towards an Actor
+ */
+USTRUCT(meta=(DisplayName="Face Towards Actor", Category="Combat"))
+struct FStateTreeFaceActorTask : public FStateTreeTaskCommonBase
+{
+	GENERATED_BODY()
+
+	/* Ensure we're using the correct instance data struct */
+	using FInstanceDataType = FStateTreeFaceActorInstanceData;
+	virtual const UStruct* GetInstanceDataType() const override { return FInstanceDataType::StaticStruct(); }
+
+	/** Runs when the owning state is entered */
+	virtual EStateTreeRunStatus EnterState(FStateTreeExecutionContext& Context, const FStateTreeTransitionResult& Transition) const override;
+
+	/** Runs when the owning state is ended */
+	virtual void ExitState(FStateTreeExecutionContext& Context, const FStateTreeTransitionResult& Transition) const override;
+
+#if WITH_EDITOR
+	virtual FText GetDescription(const FGuid& ID, FStateTreeDataView InstanceDataView, const IStateTreeBindingLookup& BindingLookup, EStateTreeNodeFormatting Formatting = EStateTreeNodeFormatting::Text) const override;
+#endif // WITH_EDITOR
+};
+
+////////////////////////////////////////////////////////////////////
+
+/**
+ *  Instance data struct for the Face Towards Location StateTree task
+ */
+USTRUCT()
+struct FStateTreeFaceLocationInstanceData
+{
+	GENERATED_BODY()
+
+	/** AI Controller that will determine the focused location */
+	UPROPERTY(EditAnywhere, Category = Context)
+	TObjectPtr<AAIController> Controller;
+
+	/** Location that will be faced towards */
+	UPROPERTY(EditAnywhere, Category = Parameter)
+	FVector FaceLocation;
+};
+
+/**
+ *  StateTree task to face an AI-Controlled Pawn towards a world location
+ */
+USTRUCT(meta=(DisplayName="Face Towards Location", Category="Combat"))
+struct FStateTreeFaceLocationTask : public FStateTreeTaskCommonBase
+{
+	GENERATED_BODY()
+
+	/* Ensure we're using the correct instance data struct */
+	using FInstanceDataType = FStateTreeFaceLocationInstanceData;
+	virtual const UStruct* GetInstanceDataType() const override { return FInstanceDataType::StaticStruct(); }
+
+	/** Runs when the owning state is entered */
+	virtual EStateTreeRunStatus EnterState(FStateTreeExecutionContext& Context, const FStateTreeTransitionResult& Transition) const override;
+
+	/** Runs when the owning state is ended */
+	virtual void ExitState(FStateTreeExecutionContext& Context, const FStateTreeTransitionResult& Transition) const override;
+
+#if WITH_EDITOR
+	virtual FText GetDescription(const FGuid& ID, FStateTreeDataView InstanceDataView, const IStateTreeBindingLookup& BindingLookup, EStateTreeNodeFormatting Formatting = EStateTreeNodeFormatting::Text) const override;
+#endif // WITH_EDITOR
+};
+
+////////////////////////////////////////////////////////////////////
+
+/**
+ *  Instance data struct for the Set Character Speed StateTree task
+ */
+USTRUCT()
+struct FStateTreeSetCharacterSpeedInstanceData
+{
+	GENERATED_BODY()
+
+	/** Character that will be affected */
+	UPROPERTY(EditAnywhere, Category = Context)
+	TObjectPtr<ACharacter> Character;
+
+	/** Max ground speed to set for the character */
+	UPROPERTY(EditAnywhere, Category = Parameter)
+	float Speed = 600.0f;
+};
+
+/**
+ *  StateTree task to change a Character's ground speed
+ */
+USTRUCT(meta=(DisplayName="Set Character Speed", Category="Combat"))
+struct FStateTreeSetCharacterSpeedTask : public FStateTreeTaskCommonBase
+{
+	GENERATED_BODY()
+
+	/* Ensure we're using the correct instance data struct */
+	using FInstanceDataType = FStateTreeSetCharacterSpeedInstanceData;
+	virtual const UStruct* GetInstanceDataType() const override { return FInstanceDataType::StaticStruct(); }
+
+	/** Runs when the owning state is entered */
+	virtual EStateTreeRunStatus EnterState(FStateTreeExecutionContext& Context, const FStateTreeTransitionResult& Transition) const override;
+
+#if WITH_EDITOR
+	virtual FText GetDescription(const FGuid& ID, FStateTreeDataView InstanceDataView, const IStateTreeBindingLookup& BindingLookup, EStateTreeNodeFormatting Formatting = EStateTreeNodeFormatting::Text) const override;
+#endif // WITH_EDITOR
+};
+
+////////////////////////////////////////////////////////////////////
+
+/**
+ *  Instance data struct for the Get Player Info task
+ */
+USTRUCT()
+struct FStateTreeGetPlayerInfoInstanceData
+{
+	GENERATED_BODY()
+
+	/** Character that owns this task */
+	UPROPERTY(EditAnywhere, Category = Context)
+	TObjectPtr<ACharacter> Character;
+
+	/** Character that owns this task */
+	UPROPERTY(VisibleAnywhere)
+	TObjectPtr<ACharacter> TargetPlayerCharacter;
+
+	/** Last known location for the target */
+	UPROPERTY(VisibleAnywhere)
+	FVector TargetPlayerLocation;
+
+	/** Distance to the target */
+	UPROPERTY(VisibleAnywhere)
+	float DistanceToTarget;
+};
+
+/**
+ *  StateTree task to get information about the player character
+ */
+USTRUCT(meta=(DisplayName="GetPlayerInfo", Category="Combat"))
+struct FStateTreeGetPlayerInfoTask : public FStateTreeTaskCommonBase
+{
+	GENERATED_BODY()
+
+	/* Ensure we're using the correct instance data struct */
+	using FInstanceDataType = FStateTreeGetPlayerInfoInstanceData;
+	virtual const UStruct* GetInstanceDataType() const override { return FInstanceDataType::StaticStruct(); }
+
+	/** Runs while the owning state is active */
+	virtual EStateTreeRunStatus Tick(FStateTreeExecutionContext& Context, const float DeltaTime) const override;
+
+#if WITH_EDITOR
+	virtual FText GetDescription(const FGuid& ID, FStateTreeDataView InstanceDataView, const IStateTreeBindingLookup& BindingLookup, EStateTreeNodeFormatting Formatting = EStateTreeNodeFormatting::Text) const override;
+#endif // WITH_EDITOR
+};
\ No newline at end of file
diff --git a/Templates/TP_ThirdPerson/Source/TP_ThirdPerson/Variant_Combat/AI/EnvQueryContext_Player.cpp b/Templates/TP_ThirdPerson/Source/TP_ThirdPerson/Variant_Combat/AI/EnvQueryContext_Player.cpp
new file mode 100644
--- /dev/null
+++ b/Templates/TP_ThirdPerson/Source/TP_ThirdPerson/Variant_Combat/AI/EnvQueryContext_Player.cpp
@@ -0,0 +1,18 @@
+// Copyright Epic Games, Inc. All Rights Reserved.
+
+
+#include "EnvQueryContext_Player.h"
+#include "Kismet/GameplayStatics.h"
+#include "EnvironmentQuery/EnvQueryTypes.h"
+#include "EnvironmentQuery/Items/EnvQueryItemType_Actor.h"
+#include "GameFramework/Pawn.h"
+
+void UEnvQueryContext_Player::ProvideContext(FEnvQueryInstance& QueryInstance, FEnvQueryContextData& ContextData) const
+{
+	// get the player pawn for the first local player
+	AActor* PlayerPawn = UGameplayStatics::GetPlayerPawn(QueryInstance.Owner.Get(), 0);
+	check(PlayerPawn);
+
+	// add the actor data to the context
+	UEnvQueryItemType_Actor::SetContextHelper(ContextData, PlayerPawn);
+}
diff --git a/Templates/TP_ThirdPerson/Source/TP_ThirdPerson/Variant_Combat/AI/EnvQueryContext_Player.h b/Templates/TP_ThirdPerson/Source/TP_ThirdPerson/Variant_Combat/AI/EnvQueryContext_Player.h
new file mode 100644
--- /dev/null
+++ b/Templates/TP_ThirdPerson/Source/TP_ThirdPerson/Variant_Combat/AI/EnvQueryContext_Player.h
@@ -0,0 +1,22 @@
+// Copyright Epic Games, Inc. All Rights Reserved.
+
+#pragma once
+
+#include "CoreMinimal.h"
+#include "EnvironmentQuery/EnvQueryContext.h"
+#include "EnvQueryContext_Player.generated.h"
+
+/**
+ *  UEnvQueryContext_Player
+ *  Basic EnvQuery Context that returns the first local player
+ */
+UCLASS()
+class UEnvQueryContext_Player : public UEnvQueryContext
+{
+	GENERATED_BODY()
+	
+public:
+
+	/** Provides the context locations or actors for this EnvQuery */
+	virtual void ProvideContext(FEnvQueryInstance& QueryInstance, FEnvQueryContextData& ContextData) const override;
+};
diff --git a/Templates/TP_ThirdPerson/Source/TP_ThirdPerson/Variant_Combat/AnimNotify_CheckChargedAttack.cpp b/Templates/TP_ThirdPerson/Source/TP_ThirdPerson/Variant_Combat/AnimNotify_CheckChargedAttack.cpp
new file mode 100644
--- /dev/null
+++ b/Templates/TP_ThirdPerson/Source/TP_ThirdPerson/Variant_Combat/AnimNotify_CheckChargedAttack.cpp
@@ -0,0 +1,21 @@
+// Copyright Epic Games, Inc. All Rights Reserved.
+
+
+#include "AnimNotify_CheckChargedAttack.h"
+#include "CombatAttacker.h"
+#include "Components/SkeletalMeshComponent.h"
+
+void UAnimNotify_CheckChargedAttack::Notify(USkeletalMeshComponent* MeshComp, UAnimSequenceBase* Animation, const FAnimNotifyEventReference& EventReference)
+{
+	// cast the owner to the attacker interface
+	if (ICombatAttacker* AttackerInterface = Cast<ICombatAttacker>(MeshComp->GetOwner()))
+	{
+		// tell the actor to check for a charged attack loop
+		AttackerInterface->CheckChargedAttack();
+	}
+}
+
+FString UAnimNotify_CheckChargedAttack::GetNotifyName_Implementation() const
+{
+	return FString("Check Charged Attack");
+}
diff --git a/Templates/TP_ThirdPerson/Source/TP_ThirdPerson/Variant_Combat/AnimNotify_CheckChargedAttack.h b/Templates/TP_ThirdPerson/Source/TP_ThirdPerson/Variant_Combat/AnimNotify_CheckChargedAttack.h
new file mode 100644
--- /dev/null
+++ b/Templates/TP_ThirdPerson/Source/TP_ThirdPerson/Variant_Combat/AnimNotify_CheckChargedAttack.h
@@ -0,0 +1,24 @@
+// Copyright Epic Games, Inc. All Rights Reserved.
+
+#pragma once
+
+#include "CoreMinimal.h"
+#include "Animation/AnimNotifies/AnimNotify.h"
+#include "AnimNotify_CheckChargedAttack.generated.h"
+
+/**
+ *  AnimNotify to perform a charged attack hold check.
+ */
+UCLASS()
+class UAnimNotify_CheckChargedAttack : public UAnimNotify
+{
+	GENERATED_BODY()
+	
+public:
+
+	/** Perform the Anim Notify */
+	virtual void Notify(USkeletalMeshComponent* MeshComp, UAnimSequenceBase* Animation, const FAnimNotifyEventReference& EventReference) override;
+
+	/** Get the notify name */
+	virtual FString GetNotifyName_Implementation() const override;
+};
diff --git a/Templates/TP_ThirdPerson/Source/TP_ThirdPerson/Variant_Combat/AnimNotify_CheckCombo.cpp b/Templates/TP_ThirdPerson/Source/TP_ThirdPerson/Variant_Combat/AnimNotify_CheckCombo.cpp
new file mode 100644
--- /dev/null
+++ b/Templates/TP_ThirdPerson/Source/TP_ThirdPerson/Variant_Combat/AnimNotify_CheckCombo.cpp
@@ -0,0 +1,21 @@
+// Copyright Epic Games, Inc. All Rights Reserved.
+
+
+#include "AnimNotify_CheckCombo.h"
+#include "CombatAttacker.h"
+#include "Components/SkeletalMeshComponent.h"
+
+void UAnimNotify_CheckCombo::Notify(USkeletalMeshComponent* MeshComp, UAnimSequenceBase* Animation, const FAnimNotifyEventReference& EventReference)
+{
+	// cast the owner to the attacker interface
+	if (ICombatAttacker* AttackerInterface = Cast<ICombatAttacker>(MeshComp->GetOwner()))
+	{
+		// tell the actor to check for combo string
+		AttackerInterface->CheckCombo();
+	}
+}
+
+FString UAnimNotify_CheckCombo::GetNotifyName_Implementation() const
+{
+	return FString("Check Combo String");
+}
diff --git a/Templates/TP_ThirdPerson/Source/TP_ThirdPerson/Variant_Combat/AnimNotify_CheckCombo.h b/Templates/TP_ThirdPerson/Source/TP_ThirdPerson/Variant_Combat/AnimNotify_CheckCombo.h
new file mode 100644
--- /dev/null
+++ b/Templates/TP_ThirdPerson/Source/TP_ThirdPerson/Variant_Combat/AnimNotify_CheckCombo.h
@@ -0,0 +1,24 @@
+// Copyright Epic Games, Inc. All Rights Reserved.
+
+#pragma once
+
+#include "CoreMinimal.h"
+#include "Animation/AnimNotifies/AnimNotify.h"
+#include "AnimNotify_CheckCombo.generated.h"
+
+/**
+ *  AnimNotify to perform a combo string check.
+ */
+UCLASS()
+class UAnimNotify_CheckCombo : public UAnimNotify
+{
+	GENERATED_BODY()
+	
+public:
+
+	/** Perform the Anim Notify */
+	virtual void Notify(USkeletalMeshComponent* MeshComp, UAnimSequenceBase* Animation, const FAnimNotifyEventReference& EventReference) override;
+
+	/** Get the notify name */
+	virtual FString GetNotifyName_Implementation() const override;
+};
diff --git a/Templates/TP_ThirdPerson/Source/TP_ThirdPerson/Variant_Combat/AnimNotify_DoAttackTrace.cpp b/Templates/TP_ThirdPerson/Source/TP_ThirdPerson/Variant_Combat/AnimNotify_DoAttackTrace.cpp
new file mode 100644
--- /dev/null
+++ b/Templates/TP_ThirdPerson/Source/TP_ThirdPerson/Variant_Combat/AnimNotify_DoAttackTrace.cpp
@@ -0,0 +1,20 @@
+// Copyright Epic Games, Inc. All Rights Reserved.
+
+
+#include "AnimNotify_DoAttackTrace.h"
+#include "CombatAttacker.h"
+#include "Components/SkeletalMeshComponent.h"
+
+void UAnimNotify_DoAttackTrace::Notify(USkeletalMeshComponent* MeshComp, UAnimSequenceBase* Animation, const FAnimNotifyEventReference& EventReference)
+{
+	// cast the owner to the attacker interface
+	if (ICombatAttacker* AttackerInterface = Cast<ICombatAttacker>(MeshComp->GetOwner()))
+	{
+		AttackerInterface->DoAttackTrace(AttackBoneName);
+	}
+}
+
+FString UAnimNotify_DoAttackTrace::GetNotifyName_Implementation() const
+{
+	return FString("Do Attack Trace");
+}
\ No newline at end of file
diff --git a/Templates/TP_ThirdPerson/Source/TP_ThirdPerson/Variant_Combat/AnimNotify_DoAttackTrace.h b/Templates/TP_ThirdPerson/Source/TP_ThirdPerson/Variant_Combat/AnimNotify_DoAttackTrace.h
new file mode 100644
--- /dev/null
+++ b/Templates/TP_ThirdPerson/Source/TP_ThirdPerson/Variant_Combat/AnimNotify_DoAttackTrace.h
@@ -0,0 +1,30 @@
+// Copyright Epic Games, Inc. All Rights Reserved.
+
+#pragma once
+
+#include "CoreMinimal.h"
+#include "Animation/AnimNotifies/AnimNotify.h"
+#include "AnimNotify_DoAttackTrace.generated.h"
+
+/**
+ *  AnimNotify to tell the actor to perform an attack trace check to look for targets to damage.
+ */
+UCLASS()
+class UAnimNotify_DoAttackTrace : public UAnimNotify
+{
+	GENERATED_BODY()
+	
+protected:
+
+	/** Source bone for the attack trace */
+	UPROPERTY(EditAnywhere, Category="Attack")
+	FName AttackBoneName;
+
+public:
+
+	/** Perform the Anim Notify */
+	virtual void Notify(USkeletalMeshComponent* MeshComp, UAnimSequenceBase* Animation, const FAnimNotifyEventReference& EventReference) override;
+
+	/** Get the notify name */
+	virtual FString GetNotifyName_Implementation() const override;
+};
diff --git a/Templates/TP_ThirdPerson/Source/TP_ThirdPerson/Variant_Combat/CombatActivatable.cpp b/Templates/TP_ThirdPerson/Source/TP_ThirdPerson/Variant_Combat/CombatActivatable.cpp
new file mode 100644
--- /dev/null
+++ b/Templates/TP_ThirdPerson/Source/TP_ThirdPerson/Variant_Combat/CombatActivatable.cpp
@@ -0,0 +1,4 @@
+// Copyright Epic Games, Inc. All Rights Reserved.
+
+
+#include "CombatActivatable.h"
\ No newline at end of file
diff --git a/Templates/TP_ThirdPerson/Source/TP_ThirdPerson/Variant_Combat/CombatActivatable.h b/Templates/TP_ThirdPerson/Source/TP_ThirdPerson/Variant_Combat/CombatActivatable.h
new file mode 100644
--- /dev/null
+++ b/Templates/TP_ThirdPerson/Source/TP_ThirdPerson/Variant_Combat/CombatActivatable.h
@@ -0,0 +1,36 @@
+// Copyright Epic Games, Inc. All Rights Reserved.
+
+#pragma once
+
+#include "CoreMinimal.h"
+#include "UObject/Interface.h"
+#include "CombatActivatable.generated.h"
+
+/**
+ *  Interactable Interface
+ *  Provides a context-agnostic way of activating, deactivating or toggling actors
+ */
+UINTERFACE(MinimalAPI, NotBlueprintable)
+class UCombatActivatable : public UInterface
+{
+	GENERATED_BODY()
+};
+
+class ICombatActivatable
+{
+	GENERATED_BODY()
+
+public:
+
+	/** Toggles the Interactable Actor */
+	UFUNCTION(BlueprintCallable, Category="Activatable")
+	virtual void ToggleInteraction(AActor* ActivationInstigator) = 0;
+
+	/** Activates the Interactable Actor */
+	UFUNCTION(BlueprintCallable, Category="Activatable")
+	virtual void ActivateInteraction(AActor* ActivationInstigator) = 0;
+
+	/** Deactivates the Interactable Actor */
+	UFUNCTION(BlueprintCallable, Category="Activatable")
+	virtual void DeactivateInteraction(AActor* ActivationInstigator) = 0;
+};
diff --git a/Templates/TP_ThirdPerson/Source/TP_ThirdPerson/Variant_Combat/CombatActivationVolume.cpp b/Templates/TP_ThirdPerson/Source/TP_ThirdPerson/Variant_Combat/CombatActivationVolume.cpp
new file mode 100644
--- /dev/null
+++ b/Templates/TP_ThirdPerson/Source/TP_ThirdPerson/Variant_Combat/CombatActivationVolume.cpp
@@ -0,0 +1,49 @@
+// Copyright Epic Games, Inc. All Rights Reserved.
+
+
+#include "CombatActivationVolume.h"
+#include "Components/BoxComponent.h"
+#include "GameFramework/Character.h"
+#include "CombatActivatable.h"
+
+ACombatActivationVolume::ACombatActivationVolume()
+{
+	PrimaryActorTick.bCanEverTick = false;
+
+	// create the box volume
+	RootComponent = Box = CreateDefaultSubobject<UBoxComponent>(TEXT("Box"));
+	check(Box);
+
+	// set the box's extent
+	Box->SetBoxExtent(FVector(500.0f, 500.0f, 500.0f));
+
+	// set the default collision profile to overlap all dynamic
+	Box->SetCollisionProfileName(FName("OverlapAllDynamic"));
+
+	// bind the begin overlap 
+	Box->OnComponentBeginOverlap.AddDynamic(this, &ACombatActivationVolume::OnOverlap);
+}
+
+void ACombatActivationVolume::OnOverlap(UPrimitiveComponent* OverlappedComponent, AActor* OtherActor, UPrimitiveComponent* OtherComp, int32 OtherBodyIndex, bool bFromSweep, const FHitResult& SweepResult)
+{
+	// has a Character entered the volume?
+	ACharacter* PlayerCharacter = Cast<ACharacter>(OtherActor);
+
+	if (PlayerCharacter)
+	{
+		// is the Character controlled by a player
+		if (PlayerCharacter->IsPlayerControlled())
+		{
+			// process the actors to activate list
+			for (AActor* CurrentActor : ActorsToActivate)
+			{
+				// is the referenced actor activatable?
+				if(ICombatActivatable* Activatable = Cast<ICombatActivatable>(CurrentActor))
+				{
+					Activatable->ActivateInteraction(PlayerCharacter);
+				}
+			}
+		}
+	}
+
+}
\ No newline at end of file
diff --git a/Templates/TP_ThirdPerson/Source/TP_ThirdPerson/Variant_Combat/CombatActivationVolume.h b/Templates/TP_ThirdPerson/Source/TP_ThirdPerson/Variant_Combat/CombatActivationVolume.h
new file mode 100644
--- /dev/null
+++ b/Templates/TP_ThirdPerson/Source/TP_ThirdPerson/Variant_Combat/CombatActivationVolume.h
@@ -0,0 +1,40 @@
+// Copyright Epic Games, Inc. All Rights Reserved.
+
+#pragma once
+
+#include "CoreMinimal.h"
+#include "GameFramework/Actor.h"
+#include "CombatActivationVolume.generated.h"
+
+class UBoxComponent;
+
+/**
+ *  A simple volume that activates a list of actors when the player pawn enters.
+ */
+UCLASS()
+class ACombatActivationVolume : public AActor
+{
+	GENERATED_BODY()
+
+	/** Collision box volume */
+	UPROPERTY(VisibleAnywhere, BlueprintReadOnly, Category = Components, meta = (AllowPrivateAccess = "true"))
+	UBoxComponent* Box;
+	
+protected:
+
+	/** List of actors to activate when this volume is entered */
+	UPROPERTY(EditAnywhere, BlueprintReadOnly, Category="Activation Volume")
+	TArray<AActor*> ActorsToActivate;
+
+public:	
+	
+	/** Constructor */
+	ACombatActivationVolume();
+
+protected:
+
+	/** Handles overlaps with the box volume */
+	UFUNCTION()
+	void OnOverlap(UPrimitiveComponent* OverlappedComponent, AActor* OtherActor, UPrimitiveComponent* OtherComp, int32 OtherBodyIndex, bool bFromSweep, const FHitResult& SweepResult);
+
+};
diff --git a/Templates/TP_ThirdPerson/Source/TP_ThirdPerson/Variant_Combat/CombatAttacker.cpp b/Templates/TP_ThirdPerson/Source/TP_ThirdPerson/Variant_Combat/CombatAttacker.cpp
new file mode 100644
--- /dev/null
+++ b/Templates/TP_ThirdPerson/Source/TP_ThirdPerson/Variant_Combat/CombatAttacker.cpp
@@ -0,0 +1,4 @@
+// Copyright Epic Games, Inc. All Rights Reserved.
+
+
+#include "CombatAttacker.h"
diff --git a/Templates/TP_ThirdPerson/Source/TP_ThirdPerson/Variant_Combat/CombatAttacker.h b/Templates/TP_ThirdPerson/Source/TP_ThirdPerson/Variant_Combat/CombatAttacker.h
new file mode 100644
--- /dev/null
+++ b/Templates/TP_ThirdPerson/Source/TP_ThirdPerson/Variant_Combat/CombatAttacker.h
@@ -0,0 +1,36 @@
+// Copyright Epic Games, Inc. All Rights Reserved.
+
+#pragma once
+
+#include "CoreMinimal.h"
+#include "UObject/Interface.h"
+#include "CombatAttacker.generated.h"
+
+/**
+ *  CombatAttacker Interface
+ *  Provides common functionality to trigger attack animation events.
+ */
+UINTERFACE(MinimalAPI, NotBlueprintable)
+class UCombatAttacker : public UInterface
+{
+	GENERATED_BODY()
+};
+
+class ICombatAttacker
+{
+	GENERATED_BODY()
+
+public:
+
+	/** Performs an attack's collision check. Usually called from a montage's AnimNotify */
+	UFUNCTION(BlueprintCallable, Category="Attacker")
+	virtual void DoAttackTrace(FName DamageSourceBone) = 0;
+
+	/** Performs a combo attack's check to continue the string. Usually called from a montage's AnimNotify */
+	UFUNCTION(BlueprintCallable, Category="Attacker")
+	virtual void CheckCombo() = 0;
+
+	/** Performs a charged attack's check to loop the charge animation. Usually called from a montage's AnimNotify */
+	UFUNCTION(BlueprintCallable, Category="Attacker")
+	virtual void CheckChargedAttack() = 0;
+};
diff --git a/Templates/TP_ThirdPerson/Source/TP_ThirdPerson/Variant_Combat/CombatCharacter.cpp b/Templates/TP_ThirdPerson/Source/TP_ThirdPerson/Variant_Combat/CombatCharacter.cpp
new file mode 100644
--- /dev/null
+++ b/Templates/TP_ThirdPerson/Source/TP_ThirdPerson/Variant_Combat/CombatCharacter.cpp
@@ -0,0 +1,488 @@
+// Copyright Epic Games, Inc. All Rights Reserved.
+
+
+#include "CombatCharacter.h"
+#include "Components/CapsuleComponent.h"
+#include "Components/WidgetComponent.h"
+#include "GameFramework/CharacterMovementComponent.h"
+#include "GameFramework/SpringArmComponent.h"
+#include "Components/SkeletalMeshComponent.h"
+#include "Camera/CameraComponent.h"
+#include "EnhancedInputSubsystems.h"
+#include "EnhancedInputComponent.h"
+#include "CombatLifeBar.h"
+#include "Engine/DamageEvents.h"
+#include "TimerManager.h"
+#include "Engine/LocalPlayer.h"
+#include "CombatPlayerController.h"
+
+DEFINE_LOG_CATEGORY(LogCombatCharacter);
+
+ACombatCharacter::ACombatCharacter()
+{
+	PrimaryActorTick.bCanEverTick = true;
+
+	// bind the attack montage ended delegate
+	OnAttackMontageEnded.BindUObject(this, &ACombatCharacter::AttackMontageEnded);
+
+	// Set size for collision capsule
+	GetCapsuleComponent()->InitCapsuleSize(35.0f, 90.0f);
+
+	// Configure character movement
+	GetCharacterMovement()->MaxWalkSpeed = 400.0f;
+
+	// create the camera boom
+	CameraBoom = CreateDefaultSubobject<USpringArmComponent>(TEXT("CameraBoom"));
+	CameraBoom->SetupAttachment(RootComponent);
+
+	CameraBoom->TargetArmLength = DefaultCameraDistance;
+	CameraBoom->bUsePawnControlRotation = true;
+	CameraBoom->bEnableCameraLag = true;
+	CameraBoom->bEnableCameraRotationLag = true;
+
+	// create the orbiting camera
+	FollowCamera = CreateDefaultSubobject<UCameraComponent>(TEXT("FollowCamera"));
+	FollowCamera->SetupAttachment(CameraBoom, USpringArmComponent::SocketName);
+	FollowCamera->bUsePawnControlRotation = false;
+
+	// create the life bar widget component
+	LifeBar = CreateDefaultSubobject<UWidgetComponent>(TEXT("LifeBar"));
+	LifeBar->SetupAttachment(RootComponent);
+
+	// set the player tag
+	Tags.Add(FName("Player"));
+}
+
+void ACombatCharacter::Move(const FInputActionValue& Value)
+{
+	// input is a Vector2D
+	FVector2D MovementVector = Value.Get<FVector2D>();
+
+	// route the input
+	DoMove(MovementVector.X, MovementVector.Y);
+}
+
+void ACombatCharacter::Look(const FInputActionValue& Value)
+{
+	FVector2D LookAxisVector = Value.Get<FVector2D>();
+
+	// route the input
+	DoLook(LookAxisVector.X, LookAxisVector.Y);
+}
+
+void ACombatCharacter::ComboAttackPressed()
+{
+	// route the input
+	DoComboAttackStart();
+}
+
+void ACombatCharacter::ChargedAttackPressed()
+{
+	// route the input
+	DoChargedAttackStart();
+}
+
+void ACombatCharacter::ChargedAttackReleased()
+{
+	// route the input
+	DoChargedAttackEnd();
+}
+
+void ACombatCharacter::DoMove(float Right, float Forward)
+{
+	if (GetController() != nullptr)
+	{
+		// find out which way is forward
+		const FRotator Rotation = GetController()->GetControlRotation();
+		const FRotator YawRotation(0, Rotation.Yaw, 0);
+
+		// get forward vector
+		const FVector ForwardDirection = FRotationMatrix(YawRotation).GetUnitAxis(EAxis::X);
+
+		// get right vector 
+		const FVector RightDirection = FRotationMatrix(YawRotation).GetUnitAxis(EAxis::Y);
+
+		// add movement 
+		AddMovementInput(ForwardDirection, Forward);
+		AddMovementInput(RightDirection, Right);
+	}
+}
+
+void ACombatCharacter::DoLook(float Yaw, float Pitch)
+{
+	if (GetController() != nullptr)
+	{
+		// add yaw and pitch input to controller
+		AddControllerYawInput(Yaw);
+		AddControllerPitchInput(Pitch);
+	}
+}
+
+void ACombatCharacter::DoComboAttackStart()
+{
+	// are we already playing an attack animation?
+	if (bIsAttacking)
+	{
+		// cache the input time so we can check it later
+		CachedAttackInputTime = GetWorld()->GetTimeSeconds();
+
+		return;
+	}
+
+	// perform a combo attack
+	ComboAttack();
+}
+
+void ACombatCharacter::DoComboAttackEnd()
+{
+	// stub
+}
+
+void ACombatCharacter::DoChargedAttackStart()
+{
+	// raise the charging attack flag
+	bIsChargingAttack = true;
+
+	if (bIsAttacking)
+	{
+		// cache the input time so we can check it later
+		CachedAttackInputTime = GetWorld()->GetTimeSeconds();
+
+		return;
+	}
+
+	ChargedAttack();
+}
+
+void ACombatCharacter::DoChargedAttackEnd()
+{
+	// lower the charging attack flag
+	bIsChargingAttack = false;
+
+	// if we've done the charge loop at least once, release the charged attack right away
+	if (bHasLoopedChargedAttack)
+	{
+		CheckChargedAttack();
+	}
+}
+
+void ACombatCharacter::ResetHP()
+{
+	// reset the current HP total
+	CurrentHP = MaxHP;
+
+	// update the life bar
+	LifeBarWidget->SetLifePercentage(1.0f);
+}
+
+void ACombatCharacter::ComboAttack()
+{
+	// raise the attacking flag
+	bIsAttacking = true;
+
+	// reset the combo count
+	ComboCount = 0;
+
+	// play the attack montage
+	if (UAnimInstance* AnimInstance = GetMesh()->GetAnimInstance())
+	{
+		const float MontageLength = AnimInstance->Montage_Play(ComboAttackMontage, 1.0f, EMontagePlayReturnType::MontageLength, 0.0f, true);
+
+		// subscribe to montage completed and interrupted events
+		if (MontageLength > 0.0f)
+		{
+			// set the end delegate for the montage
+			AnimInstance->Montage_SetEndDelegate(OnAttackMontageEnded, ComboAttackMontage);
+		}
+	}
+
+}
+
+void ACombatCharacter::ChargedAttack()
+{
+	// raise the attacking flag
+	bIsAttacking = true;
+
+	// reset the charge loop flag
+	bHasLoopedChargedAttack = false;
+
+	// play the charged attack montage
+	if (UAnimInstance* AnimInstance = GetMesh()->GetAnimInstance())
+	{
+		const float MontageLength = AnimInstance->Montage_Play(ChargedAttackMontage, 1.0f, EMontagePlayReturnType::MontageLength, 0.0f, true);
+
+		// subscribe to montage completed and interrupted events
+		if (MontageLength > 0.0f)
+		{
+			// set the end delegate for the montage
+			AnimInstance->Montage_SetEndDelegate(OnAttackMontageEnded, ChargedAttackMontage);
+		}
+	}
+}
+
+void ACombatCharacter::AttackMontageEnded(UAnimMontage* Montage, bool bInterrupted)
+{
+	// reset the attacking flag
+	bIsAttacking = false;
+
+	// check if we have a non-stale cached input
+	if (GetWorld()->GetTimeSeconds() - CachedAttackInputTime <= AttackInputCacheTimeTolerance)
+	{
+		// are we holding the charged attack button?
+		if (bIsChargingAttack)
+		{
+			// do a charged attack
+			ChargedAttack();
+		}
+		else
+		{
+			// do a regular attack
+			ComboAttack();
+		}
+	}
+}
+
+void ACombatCharacter::DoAttackTrace(FName DamageSourceBone)
+{
+	// sweep for objects in front of the character to be hit by the attack
+	TArray<FHitResult> OutHits;
+
+	// start at the provided socket location, sweep forward
+	const FVector TraceStart = GetMesh()->GetSocketLocation(DamageSourceBone);
+	const FVector TraceEnd = TraceStart + (GetActorForwardVector() * MeleeTraceDistance);
+
+	// check for pawn and world dynamic collision object types
+	FCollisionObjectQueryParams ObjectParams;
+	ObjectParams.AddObjectTypesToQuery(ECC_Pawn);
+	ObjectParams.AddObjectTypesToQuery(ECC_WorldDynamic);
+
+	// use a sphere shape for the sweep
+	FCollisionShape CollisionShape;
+	CollisionShape.SetSphere(MeleeTraceRadius);
+
+	// ignore self
+	FCollisionQueryParams QueryParams;
+	QueryParams.AddIgnoredActor(this);
+
+	if (GetWorld()->SweepMultiByObjectType(OutHits, TraceStart, TraceEnd, FQuat::Identity, ObjectParams, CollisionShape, QueryParams))
+	{
+		// iterate over each object hit
+		for (const FHitResult& CurrentHit : OutHits)
+		{
+			// check if we've hit a damageable actor
+			ICombatDamageable* Damageable = Cast<ICombatDamageable>(CurrentHit.GetActor());
+
+			if (Damageable)
+			{
+				// knock upwards and away from the impact normal
+				const FVector Impulse = (CurrentHit.ImpactNormal * -MeleeKnockbackImpulse) + (FVector::UpVector * MeleeLaunchImpulse);
+
+				// pass the damage event to the actor
+				Damageable->ApplyDamage(MeleeDamage, this, CurrentHit.ImpactPoint, Impulse);
+
+				// call the BP handler to play effects, etc.
+				DealtDamage(MeleeDamage, CurrentHit.ImpactPoint);
+			}
+		}
+	}
+}
+
+void ACombatCharacter::CheckCombo()
+{
+	// are we playing a non-charge attack animation?
+	if (bIsAttacking && !bIsChargingAttack)
+	{
+		// is the last attack input not stale?
+		if (GetWorld()->GetTimeSeconds() - CachedAttackInputTime <= ComboInputCacheTimeTolerance)
+		{
+			// consume the attack input so we don't accidentally trigger it twice
+			CachedAttackInputTime = 0.0f;
+
+			// increase the combo counter
+			++ComboCount;
+
+			// do we still have a combo section to play?
+			if (ComboCount < ComboSectionNames.Num())
+			{
+				// jump to the next combo section
+				if (UAnimInstance* AnimInstance = GetMesh()->GetAnimInstance())
+				{
+					AnimInstance->Montage_JumpToSection(ComboSectionNames[ComboCount], ComboAttackMontage);
+				}
+			}
+		}
+	}
+}
+
+void ACombatCharacter::CheckChargedAttack()
+{
+	// raise the looped charged attack flag
+	bHasLoopedChargedAttack = true;
+
+	// jump to either the loop or the attack section depending on whether we're still holding the charge button
+	if (UAnimInstance* AnimInstance = GetMesh()->GetAnimInstance())
+	{
+		AnimInstance->Montage_JumpToSection(bIsChargingAttack ? ChargeLoopSection : ChargeAttackSection, ChargedAttackMontage);
+	}
+}
+
+void ACombatCharacter::ApplyDamage(float Damage, AActor* DamageCauser, const FVector& DamageLocation, const FVector& DamageImpulse)
+{
+	// pass the damage event to the actor
+	FDamageEvent DamageEvent;
+	const float ActualDamage = TakeDamage(Damage, DamageEvent, nullptr, DamageCauser);
+
+	// only process knockback and effects if we received nonzero damage
+	if (ActualDamage > 0.0f)
+	{
+		// apply the knockback impulse
+		GetCharacterMovement()->AddImpulse(DamageImpulse, true);
+
+		// is the character ragdolling?
+		if (GetMesh()->IsSimulatingPhysics())
+		{
+			// apply an impulse to the ragdoll
+			GetMesh()->AddImpulseAtLocation(DamageImpulse * GetMesh()->GetMass(), DamageLocation);
+		}
+
+		// pass control to BP to play effects, etc.
+		ReceivedDamage(ActualDamage, DamageLocation, DamageImpulse.GetSafeNormal());
+	}
+
+}
+
+void ACombatCharacter::HandleDeath()
+{
+	// disable movement while we're dead
+	GetCharacterMovement()->DisableMovement();
+
+	// enable full ragdoll physics
+	GetMesh()->SetSimulatePhysics(true);
+
+	// hide the life bar
+	LifeBar->SetHiddenInGame(true);
+
+	// pull back the camera
+	GetCameraBoom()->TargetArmLength = DeathCameraDistance;
+
+	// schedule respawning
+	GetWorld()->GetTimerManager().SetTimer(RespawnTimer, this, &ACombatCharacter::RespawnCharacter, RespawnTime, false);
+}
+
+void ACombatCharacter::ApplyHealing(float Healing, AActor* Healer)
+{
+	// stub
+}
+
+void ACombatCharacter::RespawnCharacter()
+{
+	// destroy the character and let it be respawned by the Player Controller
+	Destroy();
+}
+
+float ACombatCharacter::TakeDamage(float Damage, struct FDamageEvent const& DamageEvent, AController* EventInstigator, AActor* DamageCauser)
+{
+	// only process damage if the character is still alive
+	if (CurrentHP <= 0.0f)
+	{
+		return 0.0f;
+	}
+
+	// reduce the current HP
+	CurrentHP -= Damage;
+
+	// have we run out of HP?
+	if (CurrentHP <= 0.0f)
+	{
+		// die
+		HandleDeath();
+	}
+	else
+	{
+		// update the life bar
+		LifeBarWidget->SetLifePercentage(CurrentHP / MaxHP);
+
+		// enable partial ragdoll physics, but keep the pelvis vertical
+		GetMesh()->SetPhysicsBlendWeight(0.5f);
+		GetMesh()->SetBodySimulatePhysics(PelvisBoneName, false);
+	}
+
+	// return the received damage amount
+	return Damage;
+}
+
+void ACombatCharacter::Landed(const FHitResult& Hit)
+{
+	Super::Landed(Hit);
+
+	// is the character still alive?
+	if (CurrentHP >= 0.0f)
+	{
+		// disable ragdoll physics
+		GetMesh()->SetPhysicsBlendWeight(0.0f);
+	}
+}
+
+void ACombatCharacter::BeginPlay()
+{
+	Super::BeginPlay();
+
+	// get the life bar from the widget component
+	LifeBarWidget = Cast<UCombatLifeBar>(LifeBar->GetUserWidgetObject());
+	check(LifeBarWidget);
+
+	// initialize the camera
+	GetCameraBoom()->TargetArmLength = DefaultCameraDistance;
+
+	// save the relative transform for the mesh so we can reset the ragdoll later
+	MeshStartingTransform = GetMesh()->GetRelativeTransform();
+
+	// set the life bar color
+	LifeBarWidget->SetBarColor(LifeBarColor);
+
+	// reset HP to maximum
+	ResetHP();
+}
+
+void ACombatCharacter::EndPlay(const EEndPlayReason::Type EndPlayReason)
+{
+	Super::EndPlay(EndPlayReason);
+
+	// clear the respawn timer
+	GetWorld()->GetTimerManager().ClearTimer(RespawnTimer);
+}
+
+void ACombatCharacter::SetupPlayerInputComponent(UInputComponent* PlayerInputComponent)
+{
+	Super::SetupPlayerInputComponent(PlayerInputComponent);
+
+	// Set up action bindings
+	if (UEnhancedInputComponent* EnhancedInputComponent = Cast<UEnhancedInputComponent>(PlayerInputComponent))
+	{
+		// Moving
+		EnhancedInputComponent->BindAction(MoveAction, ETriggerEvent::Triggered, this, &ACombatCharacter::Move);
+
+		// Looking
+		EnhancedInputComponent->BindAction(LookAction, ETriggerEvent::Triggered, this, &ACombatCharacter::Look);
+		EnhancedInputComponent->BindAction(MouseLookAction, ETriggerEvent::Triggered, this, &ACombatCharacter::Look);
+
+		// Combo Attack
+		EnhancedInputComponent->BindAction(ComboAttackAction, ETriggerEvent::Started, this, &ACombatCharacter::ComboAttackPressed);
+
+		// Charged Attack
+		EnhancedInputComponent->BindAction(ChargedAttackAction, ETriggerEvent::Started, this, &ACombatCharacter::ChargedAttackPressed);
+		EnhancedInputComponent->BindAction(ChargedAttackAction, ETriggerEvent::Completed, this, &ACombatCharacter::ChargedAttackReleased);
+	}
+}
+
+void ACombatCharacter::NotifyControllerChanged()
+{
+	Super::NotifyControllerChanged();
+
+	// update the respawn transform on the Player Controller
+	if (ACombatPlayerController* PC = Cast<ACombatPlayerController>(GetController()))
+	{
+		PC->SetRespawnTransform(GetActorTransform());
+	}
+}
+
diff --git a/Templates/TP_ThirdPerson/Source/TP_ThirdPerson/Variant_Combat/CombatCharacter.h b/Templates/TP_ThirdPerson/Source/TP_ThirdPerson/Variant_Combat/CombatCharacter.h
new file mode 100644
--- /dev/null
+++ b/Templates/TP_ThirdPerson/Source/TP_ThirdPerson/Variant_Combat/CombatCharacter.h
@@ -0,0 +1,309 @@
+// Copyright Epic Games, Inc. All Rights Reserved.
+
+#pragma once
+
+#include "CoreMinimal.h"
+#include "GameFramework/Character.h"
+#include "CombatAttacker.h"
+#include "CombatDamageable.h"
+#include "Animation/AnimInstance.h"
+#include "CombatCharacter.generated.h"
+
+class USpringArmComponent;
+class UCameraComponent;
+class UInputAction;
+struct FInputActionValue;
+class UCombatLifeBar;
+class UWidgetComponent;
+
+DECLARE_LOG_CATEGORY_EXTERN(LogCombatCharacter, Log, All);
+
+/**
+ *  An enhanced Third Person Character with melee combat capabilities:
+ *  - Combo attack string
+ *  - Press and hold charged attack
+ *  - Damage dealing and reaction
+ *  - Death
+ *  - Respawning
+ */
+UCLASS(abstract)
+class ACombatCharacter : public ACharacter, public ICombatAttacker, public ICombatDamageable
+{
+	GENERATED_BODY()
+
+	/** Camera boom positioning the camera behind the character */
+	UPROPERTY(VisibleAnywhere, BlueprintReadOnly, Category = Camera, meta = (AllowPrivateAccess = "true"))
+	USpringArmComponent* CameraBoom;
+
+	/** Follow camera */
+	UPROPERTY(VisibleAnywhere, BlueprintReadOnly, Category = Camera, meta = (AllowPrivateAccess = "true"))
+	UCameraComponent* FollowCamera;
+
+	/** Life bar widget component */
+	UPROPERTY(VisibleAnywhere, BlueprintReadOnly, Category = UI, meta = (AllowPrivateAccess = "true"))
+	UWidgetComponent* LifeBar;
+	
+protected:
+
+	/** Jump Input Action */
+	UPROPERTY(EditAnywhere, Category ="Input")
+	UInputAction* JumpAction;
+
+	/** Move Input Action */
+	UPROPERTY(EditAnywhere, Category ="Input")
+	UInputAction* MoveAction;
+
+	/** Look Input Action */
+	UPROPERTY(EditAnywhere, Category ="Input")
+	UInputAction* LookAction;
+
+	/** Mouse Look Input Action */
+	UPROPERTY(EditAnywhere, BlueprintReadOnly, Category="Input")
+	UInputAction* MouseLookAction;
+
+	/** Combo Attack Input Action */
+	UPROPERTY(EditAnywhere, Category ="Input")
+	UInputAction* ComboAttackAction;
+
+	/** Charged Attack Input Action */
+	UPROPERTY(EditAnywhere, Category ="Input")
+	UInputAction* ChargedAttackAction;
+
+	/** Max amount of HP the character will have on respawn */
+	UPROPERTY(EditAnywhere, Category="Damage", meta = (ClampMin = 0, ClampMax = 100))
+	float MaxHP = 5.0f;
+
+	/** Current amount of HP the character has */
+	UPROPERTY(VisibleAnywhere, Category="Damage")
+	float CurrentHP = 0.0f;
+
+	/** Life bar widget fill color */
+	UPROPERTY(EditAnywhere, Category="Damage")
+	FLinearColor LifeBarColor;
+
+	/** Name of the pelvis bone, for damage ragdoll physics */
+	UPROPERTY(EditAnywhere, Category="Damage")
+	FName PelvisBoneName;
+
+	/** Pointer to the life bar widget */
+	UPROPERTY(EditAnywhere, Category="Damage")
+	TObjectPtr<UCombatLifeBar> LifeBarWidget;
+
+	/** Max amount of time that may elapse for a non-combo attack input to not be considered stale */
+	UPROPERTY(EditAnywhere, Category="Melee Attack", meta = (ClampMin = 0, ClampMax = 5))
+	float AttackInputCacheTimeTolerance = 1.0f;
+
+	/** Time at which an attack button was last pressed */
+	float CachedAttackInputTime = 0.0f;
+
+	/** If true, the character is currently playing an attack animation */
+	bool bIsAttacking = false;
+
+	/** Distance ahead of the character that melee attack sphere collision traces will extend */
+	UPROPERTY(EditAnywhere, Category="Melee Attack|Trace", meta=(ClampMin=0, ClampMax=500, Units="cm"))
+	float MeleeTraceDistance = 75.0f;
+
+	/** Radius of the sphere trace for melee attacks */
+	UPROPERTY(EditAnywhere, Category="Melee Attack|Trace", meta = (ClampMin = 0, ClampMax = 200, Units = "cm"))
+	float MeleeTraceRadius = 75.0f;
+
+	/** Amount of damage a melee attack will deal */
+	UPROPERTY(EditAnywhere, Category="Melee Attack|Damage", meta = (ClampMin = 0, ClampMax = 100))
+	float MeleeDamage = 1.0f;
+
+	/** Amount of knockback impulse a melee attack will apply */
+	UPROPERTY(EditAnywhere, Category="Melee Attack|Damage", meta = (ClampMin = 0, ClampMax = 1000, Units = "cm/s"))
+	float MeleeKnockbackImpulse = 250.0f;
+
+	/** Amount of upwards impulse a melee attack will apply */
+	UPROPERTY(EditAnywhere, Category="Melee Attack|Damage", meta = (ClampMin = 0, ClampMax = 1000, Units = "cm/s"))
+	float MeleeLaunchImpulse = 300.0f;
+
+	/** AnimMontage that will play for combo attacks */
+	UPROPERTY(EditAnywhere, Category="Melee Attack|Combo")
+	UAnimMontage* ComboAttackMontage;
+
+	/** Names of the AnimMontage sections that correspond to each stage of the combo attack */
+	UPROPERTY(EditAnywhere, Category="Melee Attack|Combo")
+	TArray<FName> ComboSectionNames;
+
+	/** Max amount of time that may elapse for a combo attack input to not be considered stale */
+	UPROPERTY(EditAnywhere, Category="Melee Attack|Combo", meta = (ClampMin = 0, ClampMax = 5))
+	float ComboInputCacheTimeTolerance = 0.45f;
+
+	/** Index of the current stage of the melee attack combo */
+	int32 ComboCount = 0;
+
+	/** AnimMontage that will play for charged attacks */
+	UPROPERTY(EditAnywhere, Category="Melee Attack|Charged")
+	UAnimMontage* ChargedAttackMontage;
+
+	/** Name of the AnimMontage section that corresponds to the charge loop */
+	UPROPERTY(EditAnywhere, Category="Melee Attack|Charged")
+	FName ChargeLoopSection;
+
+	/** Name of the AnimMontage section that corresponds to the attack */
+	UPROPERTY(EditAnywhere, Category="Melee Attack|Charged")
+	FName ChargeAttackSection;
+
+	/** Flag that determines if the player is currently holding the charged attack input */
+	bool bIsChargingAttack = false;
+	
+	/** If true, the charged attack hold check has been tested at least once */
+	bool bHasLoopedChargedAttack = false;
+
+	/** Camera boom length while the character is dead */
+	UPROPERTY(EditAnywhere, Category="Camera", meta = (ClampMin = 0, ClampMax = 1000, Units = "cm"))
+	float DeathCameraDistance = 400.0f;
+
+	/** Camera boom length when the character respawns */
+	UPROPERTY(EditAnywhere, Category="Camera", meta = (ClampMin = 0, ClampMax = 1000, Units = "cm"))
+	float DefaultCameraDistance = 100.0f;
+
+	/** Time to wait before respawning the character */
+	UPROPERTY(EditAnywhere, Category="Respawn", meta = (ClampMin = 0, ClampMax = 10))
+	float RespawnTime = 3.0f;
+
+	/** Attack montage ended delegate */
+	FOnMontageEnded OnAttackMontageEnded;
+
+	/** Character respawn timer */
+	FTimerHandle RespawnTimer;
+
+	/** Copy of the mesh's transform so we can reset it after ragdoll animations */
+	FTransform MeshStartingTransform;
+
+public:
+	
+	/** Constructor */
+	ACombatCharacter();
+
+protected:
+
+	/** Called for movement input */
+	void Move(const FInputActionValue& Value);
+
+	/** Called for looking input */
+	void Look(const FInputActionValue& Value);
+
+	/** Called for combo attack input */
+	void ComboAttackPressed();
+
+	/** Called for combo attack input pressed */
+	void ChargedAttackPressed();
+
+	/** Called for combo attack input released */
+	void ChargedAttackReleased();
+
+public:
+
+	/** Handles move inputs from either controls or UI interfaces */
+	UFUNCTION(BlueprintCallable, Category="Input")
+	virtual void DoMove(float Right, float Forward);
+
+	/** Handles look inputs from either controls or UI interfaces */
+	UFUNCTION(BlueprintCallable, Category="Input")
+	virtual void DoLook(float Yaw, float Pitch);
+
+	/** Handles combo attack pressed from either controls or UI interfaces */
+	UFUNCTION(BlueprintCallable, Category="Input")
+	virtual void DoComboAttackStart();
+
+	/** Handles combo attack released from either controls or UI interfaces */
+	UFUNCTION(BlueprintCallable, Category="Input")
+	virtual void DoComboAttackEnd();
+
+	/** Handles charged attack pressed from either controls or UI interfaces */
+	UFUNCTION(BlueprintCallable, Category="Input")
+	virtual void DoChargedAttackStart();
+
+	/** Handles charged attack released from either controls or UI interfaces */
+	UFUNCTION(BlueprintCallable, Category="Input")
+	virtual void DoChargedAttackEnd();
+
+protected:
+
+	/** Resets the character's current HP to maximum */
+	void ResetHP();
+
+	/** Performs a combo attack */
+	void ComboAttack();
+
+	/** Performs a charged attack */
+	void ChargedAttack();
+
+	/** Called from a delegate when the attack montage ends */
+	void AttackMontageEnded(UAnimMontage* Montage, bool bInterrupted);
+
+	
+public:
+
+	// ~begin CombatAttacker interface
+
+	/** Performs the collision check for an attack */
+	virtual void DoAttackTrace(FName DamageSourceBone) override;
+
+	/** Performs the combo string check */
+	virtual void CheckCombo() override;
+
+	/** Performs the charged attack hold check */
+	virtual void CheckChargedAttack() override;
+
+	// ~end CombatAttacker interface
+
+	// ~begin CombatDamageable interface
+
+	/** Handles damage and knockback events */
+	virtual void ApplyDamage(float Damage, AActor* DamageCauser, const FVector& DamageLocation, const FVector& DamageImpulse) override;
+
+	/** Handles death events */
+	virtual void HandleDeath() override;
+
+	/** Handles healing events */
+	virtual void ApplyHealing(float Healing, AActor* Healer) override;
+
+	// ~end CombatDamageable interface
+
+	/** Called from the respawn timer to destroy and re-create the character */
+	void RespawnCharacter();
+
+public:
+
+	/** Overrides the default TakeDamage functionality */
+	virtual float TakeDamage(float Damage, struct FDamageEvent const& DamageEvent, AController* EventInstigator, AActor* DamageCauser) override;
+
+	/** Overrides landing to reset damage ragdoll physics */
+	virtual void Landed(const FHitResult& Hit) override;
+
+protected:
+
+	/** Blueprint handler to play damage dealt effects */
+	UFUNCTION(BlueprintImplementableEvent, Category="Combat")
+	void DealtDamage(float Damage, const FVector& ImpactPoint);
+
+	/** Blueprint handler to play damage received effects */
+	UFUNCTION(BlueprintImplementableEvent, Category="Combat")
+	void ReceivedDamage(float Damage, const FVector& ImpactPoint, const FVector& DamageDirection);
+
+protected:
+
+	/** Initialization */
+	virtual void BeginPlay() override;
+
+	/** Cleanup */
+	virtual void EndPlay(const EEndPlayReason::Type EndPlayReason) override;
+
+	/** Handles input bindings */
+	virtual void SetupPlayerInputComponent(class UInputComponent* PlayerInputComponent) override;
+
+	/** Handles possessed initialization */
+	virtual void NotifyControllerChanged() override;
+
+public:
+
+	/** Returns CameraBoom subobject **/
+	FORCEINLINE class USpringArmComponent* GetCameraBoom() const { return CameraBoom; }
+
+	/** Returns FollowCamera subobject **/
+	FORCEINLINE class UCameraComponent* GetFollowCamera() const { return FollowCamera; }
+};
diff --git a/Templates/TP_ThirdPerson/Source/TP_ThirdPerson/Variant_Combat/CombatCheckpointVolume.cpp b/Templates/TP_ThirdPerson/Source/TP_ThirdPerson/Variant_Combat/CombatCheckpointVolume.cpp
new file mode 100644
--- /dev/null
+++ b/Templates/TP_ThirdPerson/Source/TP_ThirdPerson/Variant_Combat/CombatCheckpointVolume.cpp
@@ -0,0 +1,47 @@
+// Copyright Epic Games, Inc. All Rights Reserved.
+
+
+#include "CombatCheckpointVolume.h"
+#include "CombatCharacter.h"
+#include "CombatPlayerController.h"
+
+ACombatCheckpointVolume::ACombatCheckpointVolume()
+{
+	// create the box volume
+	RootComponent = Box = CreateDefaultSubobject<UBoxComponent>(TEXT("Box"));
+	check(Box);
+
+	// set the box's extent
+	Box->SetBoxExtent(FVector(500.0f, 500.0f, 500.0f));
+
+	// set the default collision profile to overlap all dynamic
+	Box->SetCollisionProfileName(FName("OverlapAllDynamic"));
+
+	// bind the begin overlap 
+	Box->OnComponentBeginOverlap.AddDynamic(this, &ACombatCheckpointVolume::OnOverlap);
+}
+
+void ACombatCheckpointVolume::OnOverlap(UPrimitiveComponent* OverlappedComponent, AActor* OtherActor, UPrimitiveComponent* OtherComp, int32 OtherBodyIndex, bool bFromSweep, const FHitResult& SweepResult)
+{
+	// ensure we use this only once
+	if (bCheckpointUsed)
+	{
+		return;
+	}
+		
+	// has the player entered this volume?
+	ACombatCharacter* PlayerCharacter = Cast<ACombatCharacter>(OtherActor);
+
+	if (PlayerCharacter)
+	{
+		if (ACombatPlayerController* PC = Cast<ACombatPlayerController>(PlayerCharacter->GetController()))
+		{
+			// raise the checkpoint used flag
+			bCheckpointUsed = true;
+
+			// update the player's respawn checkpoint
+			PC->SetRespawnTransform(PlayerCharacter->GetActorTransform());
+		}
+
+	}
+}
diff --git a/Templates/TP_ThirdPerson/Source/TP_ThirdPerson/Variant_Combat/CombatCheckpointVolume.h b/Templates/TP_ThirdPerson/Source/TP_ThirdPerson/Variant_Combat/CombatCheckpointVolume.h
new file mode 100644
--- /dev/null
+++ b/Templates/TP_ThirdPerson/Source/TP_ThirdPerson/Variant_Combat/CombatCheckpointVolume.h
@@ -0,0 +1,32 @@
+// Copyright Epic Games, Inc. All Rights Reserved.
+
+#pragma once
+
+#include "CoreMinimal.h"
+#include "GameFramework/Actor.h"
+#include "Components/BoxComponent.h"
+#include "CombatCheckpointVolume.generated.h"
+
+UCLASS(abstract)
+class ACombatCheckpointVolume : public AActor
+{
+	GENERATED_BODY()
+	
+	/** Collision box volume */
+	UPROPERTY(VisibleAnywhere, BlueprintReadOnly, Category = Components, meta = (AllowPrivateAccess = "true"))
+	UBoxComponent* Box;
+
+public:	
+	
+	/** Constructor */
+	ACombatCheckpointVolume();
+
+protected:
+
+	/** Set to true after use to avoid accidentally resetting the checkpoint */
+	bool bCheckpointUsed = false;
+
+	/** Handles overlaps with the box volume */
+	UFUNCTION()
+	void OnOverlap(UPrimitiveComponent* OverlappedComponent, AActor* OtherActor, UPrimitiveComponent* OtherComp, int32 OtherBodyIndex, bool bFromSweep, const FHitResult& SweepResult);
+};
diff --git a/Templates/TP_ThirdPerson/Source/TP_ThirdPerson/Variant_Combat/CombatDamageable.cpp b/Templates/TP_ThirdPerson/Source/TP_ThirdPerson/Variant_Combat/CombatDamageable.cpp
new file mode 100644
--- /dev/null
+++ b/Templates/TP_ThirdPerson/Source/TP_ThirdPerson/Variant_Combat/CombatDamageable.cpp
@@ -0,0 +1,6 @@
+// Copyright Epic Games, Inc. All Rights Reserved.
+
+
+#include "CombatDamageable.h"
+
+// Add default functionality here for any ICombatDamageable functions that are not pure virtual.
diff --git a/Templates/TP_ThirdPerson/Source/TP_ThirdPerson/Variant_Combat/CombatDamageable.h b/Templates/TP_ThirdPerson/Source/TP_ThirdPerson/Variant_Combat/CombatDamageable.h
new file mode 100644
--- /dev/null
+++ b/Templates/TP_ThirdPerson/Source/TP_ThirdPerson/Variant_Combat/CombatDamageable.h
@@ -0,0 +1,36 @@
+// Copyright Epic Games, Inc. All Rights Reserved.
+
+#pragma once
+
+#include "CoreMinimal.h"
+#include "UObject/Interface.h"
+#include "CombatDamageable.generated.h"
+
+/**
+ *  CombatDamageable interface
+ *  Provides functionality to handle damage, healing, knockback and death
+ */
+UINTERFACE(MinimalAPI, NotBlueprintable)
+class UCombatDamageable : public UInterface
+{
+	GENERATED_BODY()
+};
+
+class ICombatDamageable
+{
+	GENERATED_BODY()
+
+public:
+
+	/** Handles damage and knockback events */
+	UFUNCTION(BlueprintCallable, Category="Damageable")
+	virtual void ApplyDamage(float Damage, AActor* DamageCauser, const FVector& DamageLocation, const FVector& DamageImpulse) = 0;
+
+	/** Handles death events */
+	UFUNCTION(BlueprintCallable, Category="Damageable")
+	virtual void HandleDeath() = 0;
+
+	/** Handles healing events */
+	UFUNCTION(BlueprintCallable, Category="Damageable")
+	virtual void ApplyHealing(float Healing, AActor* Healer) = 0;
+};
diff --git a/Templates/TP_ThirdPerson/Source/TP_ThirdPerson/Variant_Combat/CombatDamageableBox.cpp b/Templates/TP_ThirdPerson/Source/TP_ThirdPerson/Variant_Combat/CombatDamageableBox.cpp
new file mode 100644
--- /dev/null
+++ b/Templates/TP_ThirdPerson/Source/TP_ThirdPerson/Variant_Combat/CombatDamageableBox.cpp
@@ -0,0 +1,78 @@
+// Copyright Epic Games, Inc. All Rights Reserved.
+
+
+#include "CombatDamageableBox.h"
+#include "Components/StaticMeshComponent.h"
+#include "TimerManager.h"
+#include "Engine/World.h"
+
+ACombatDamageableBox::ACombatDamageableBox()
+{
+	PrimaryActorTick.bCanEverTick = false;
+
+	// create the mesh
+	RootComponent = Mesh = CreateDefaultSubobject<UStaticMeshComponent>(TEXT("Mesh"));
+
+	// set the collision properties
+	Mesh->SetCollisionProfileName(FName("BlockAllDynamic"));
+
+	// enable physics
+	Mesh->SetSimulatePhysics(true);
+
+	// disable navigation relevance so boxes don't affect NavMesh generation
+	Mesh->bNavigationRelevant = false;
+}
+
+void ACombatDamageableBox::RemoveFromLevel()
+{
+	// destroy this actor
+	Destroy();
+}
+
+void ACombatDamageableBox::EndPlay(EEndPlayReason::Type EndPlayReason)
+{
+	Super::EndPlay(EndPlayReason);
+
+	// clear the death timer
+	GetWorld()->GetTimerManager().ClearTimer(DeathTimer);
+}
+
+void ACombatDamageableBox::ApplyDamage(float Damage, AActor* DamageCauser, const FVector& DamageLocation, const FVector& DamageImpulse)
+{
+	// only process damage if we still have HP
+	if (CurrentHP > 0.0f)
+	{
+		// apply the damage
+		CurrentHP -= Damage;
+
+		// are we dead?
+		if (CurrentHP <= 0.0f)
+		{
+			HandleDeath();
+		}
+
+		// apply a physics impulse to the box, ignoring its mass
+		Mesh->AddImpulseAtLocation(DamageImpulse * Mesh->GetMass(), DamageLocation);
+
+		// call the BP handler to play effects, etc.
+		OnBoxDamaged(DamageLocation, DamageImpulse);
+	}
+}
+
+void ACombatDamageableBox::HandleDeath()
+{
+	// change the collision object type to Visibility so we ignore most interactions but still retain physics collisions
+	Mesh->SetCollisionObjectType(ECC_Visibility);
+
+	// call the BP handler to play effects, etc.
+	OnBoxDestroyed();
+
+	// set up the death cleanup timer
+	GetWorld()->GetTimerManager().SetTimer(DeathTimer, this, &ACombatDamageableBox::RemoveFromLevel, DeathDelayTime);
+}
+
+void ACombatDamageableBox::ApplyHealing(float Healing, AActor* Healer)
+{
+	// stub
+}
+
diff --git a/Templates/TP_ThirdPerson/Source/TP_ThirdPerson/Variant_Combat/CombatDamageableBox.h b/Templates/TP_ThirdPerson/Source/TP_ThirdPerson/Variant_Combat/CombatDamageableBox.h
new file mode 100644
--- /dev/null
+++ b/Templates/TP_ThirdPerson/Source/TP_ThirdPerson/Variant_Combat/CombatDamageableBox.h
@@ -0,0 +1,67 @@
+// Copyright Epic Games, Inc. All Rights Reserved.
+
+#pragma once
+
+#include "CoreMinimal.h"
+#include "GameFramework/Actor.h"
+#include "CombatDamageable.h"
+#include "CombatDamageableBox.generated.h"
+
+/**
+ *  A simple physics box that reacts to damage through the ICombatDamageable interface
+ */
+UCLASS(abstract)
+class ACombatDamageableBox : public AActor, public ICombatDamageable
+{
+	GENERATED_BODY()
+	
+	/** Damageable box mesh */
+	UPROPERTY(VisibleAnywhere, BlueprintReadOnly, Category = "Components", meta = (AllowPrivateAccess = "true"))
+	UStaticMeshComponent* Mesh;
+
+public:	
+
+	/** Constructor */
+	ACombatDamageableBox();
+
+protected:
+
+	/** Amount of HP this box starts with. */
+	UPROPERTY(EditAnywhere, BlueprintReadOnly, Category="Damage")
+	float CurrentHP = 3.0f;
+
+	/** Time to wait before we remove this box from the level. */
+	UPROPERTY(EditAnywhere, BlueprintReadOnly, Category="Damage")
+	float DeathDelayTime = 6.0f;
+
+	FTimerHandle DeathTimer;
+
+	/** Blueprint damage handler for effect playback */
+	UFUNCTION(BlueprintImplementableEvent, Category="Damage")
+	void OnBoxDamaged(const FVector& DamageLocation, const FVector& DamageImpulse);
+
+	/** Blueprint destruction handler for effect playback */
+	UFUNCTION(BlueprintImplementableEvent, Category="Damage")
+	void OnBoxDestroyed();
+
+	/** Timer callback to remove the box from the level after it dies */
+	void RemoveFromLevel();
+
+public:
+
+	/** EndPlay cleanup */
+	void EndPlay(EEndPlayReason::Type EndPlayReason) override;
+
+	// ~Begin CombatDamageable interface
+
+	/** Handles damage and knockback events */
+	virtual void ApplyDamage(float Damage, AActor* DamageCauser, const FVector& DamageLocation, const FVector& DamageImpulse) override;
+
+	/** Handles death events */
+	virtual void HandleDeath() override;
+
+	/** Handles healing events */
+	virtual void ApplyHealing(float Healing, AActor* Healer) override;
+
+	// ~End CombatDamageable interface
+};
diff --git a/Templates/TP_ThirdPerson/Source/TP_ThirdPerson/Variant_Combat/CombatDummy.cpp b/Templates/TP_ThirdPerson/Source/TP_ThirdPerson/Variant_Combat/CombatDummy.cpp
new file mode 100644
--- /dev/null
+++ b/Templates/TP_ThirdPerson/Source/TP_ThirdPerson/Variant_Combat/CombatDummy.cpp
@@ -0,0 +1,50 @@
+// Copyright Epic Games, Inc. All Rights Reserved.
+
+
+#include "Variant_Combat/CombatDummy.h"
+#include "Components/SceneComponent.h"
+#include "Components/StaticMeshComponent.h"
+#include "PhysicsEngine/PhysicsConstraintComponent.h"
+
+ACombatDummy::ACombatDummy()
+{
+ 	PrimaryActorTick.bCanEverTick = true;
+
+	// create the root
+	RootComponent = Root = CreateDefaultSubobject<USceneComponent>(TEXT("Root"));
+
+	// create the base plate
+	BasePlate = CreateDefaultSubobject<UStaticMeshComponent>(TEXT("Base Plate"));
+	BasePlate->SetupAttachment(RootComponent);
+
+	// create the dummy
+	Dummy = CreateDefaultSubobject<UStaticMeshComponent>(TEXT("Dummy"));
+	Dummy->SetupAttachment(RootComponent);
+
+	Dummy->SetSimulatePhysics(true);
+
+	// create the physics constraint
+	PhysicsConstraint = CreateEditorOnlyDefaultSubobject<UPhysicsConstraintComponent>(TEXT("Physics Constraint"));
+	PhysicsConstraint->SetupAttachment(RootComponent);
+
+	PhysicsConstraint->SetConstrainedComponents(BasePlate, NAME_None, Dummy, NAME_None);
+}
+
+void ACombatDummy::ApplyDamage(float Damage, AActor* DamageCauser, const FVector& DamageLocation, const FVector& DamageImpulse)
+{
+	// apply impulse to the dummy
+	Dummy->AddImpulseAtLocation(DamageImpulse, DamageLocation);
+
+	// call the BP handler
+	BP_OnDummyDamaged(DamageLocation, DamageImpulse.GetSafeNormal());
+}
+
+void ACombatDummy::HandleDeath()
+{
+	// unused
+}
+
+void ACombatDummy::ApplyHealing(float Healing, AActor* Healer)
+{
+	// unused
+}
\ No newline at end of file
diff --git a/Templates/TP_ThirdPerson/Source/TP_ThirdPerson/Variant_Combat/CombatDummy.h b/Templates/TP_ThirdPerson/Source/TP_ThirdPerson/Variant_Combat/CombatDummy.h
new file mode 100644
--- /dev/null
+++ b/Templates/TP_ThirdPerson/Source/TP_ThirdPerson/Variant_Combat/CombatDummy.h
@@ -0,0 +1,60 @@
+// Copyright Epic Games, Inc. All Rights Reserved.
+
+#pragma once
+
+#include "CoreMinimal.h"
+#include "GameFramework/Actor.h"
+#include "CombatDamageable.h"
+#include "CombatDummy.generated.h"
+
+class UStaticMeshComponent;
+class UPhysicsConstraintComponent;
+
+/**
+ *  A simple invincible combat training dummy
+ */
+UCLASS(abstract)
+class ACombatDummy : public AActor, public ICombatDamageable
+{
+	GENERATED_BODY()
+	
+	/** Root component */
+	UPROPERTY(VisibleAnywhere, BlueprintReadOnly, Category = "Components", meta = (AllowPrivateAccess = "true"))
+	USceneComponent* Root;
+
+	/** Static base plate */
+	UPROPERTY(VisibleAnywhere, BlueprintReadOnly, Category = "Components", meta = (AllowPrivateAccess = "true"))
+	UStaticMeshComponent* BasePlate;
+
+	/** Physics enabled dummy mesh */
+	UPROPERTY(VisibleAnywhere, BlueprintReadOnly, Category = "Components", meta = (AllowPrivateAccess = "true"))
+	UStaticMeshComponent* Dummy;
+
+	/** Physics constraint holding the dummy and base plate together */
+	UPROPERTY(VisibleAnywhere, BlueprintReadOnly, Category = "Components", meta = (AllowPrivateAccess = "true"))
+	UPhysicsConstraintComponent* PhysicsConstraint;
+
+public:	
+	
+	/** Constructor */
+	ACombatDummy();
+
+	// ~Begin CombatDamageable interface
+
+		/** Handles damage and knockback events */
+	virtual void ApplyDamage(float Damage, AActor* DamageCauser, const FVector& DamageLocation, const FVector& DamageImpulse) override;
+
+	/** Handles death events */
+	virtual void HandleDeath() override;
+
+	/** Handles healing events */
+	virtual void ApplyHealing(float Healing, AActor* Healer) override;
+
+	// ~End CombatDamageable interface
+
+protected:
+
+	/** Blueprint handle to apply damage effects */
+	UFUNCTION(BlueprintImplementableEvent, Category="Combat", meta=(DisplayName = "On Dummy Damaged"))
+	void BP_OnDummyDamaged(const FVector& Location, const FVector& Direction);
+};
diff --git a/Templates/TP_ThirdPerson/Source/TP_ThirdPerson/Variant_Combat/CombatGameMode.cpp b/Templates/TP_ThirdPerson/Source/TP_ThirdPerson/Variant_Combat/CombatGameMode.cpp
new file mode 100644
--- /dev/null
+++ b/Templates/TP_ThirdPerson/Source/TP_ThirdPerson/Variant_Combat/CombatGameMode.cpp
@@ -0,0 +1,9 @@
+// Copyright Epic Games, Inc. All Rights Reserved.
+
+
+#include "Variant_Combat/CombatGameMode.h"
+
+ACombatGameMode::ACombatGameMode()
+{
+
+}
\ No newline at end of file
diff --git a/Templates/TP_ThirdPerson/Source/TP_ThirdPerson/Variant_Combat/CombatGameMode.h b/Templates/TP_ThirdPerson/Source/TP_ThirdPerson/Variant_Combat/CombatGameMode.h
new file mode 100644
--- /dev/null
+++ b/Templates/TP_ThirdPerson/Source/TP_ThirdPerson/Variant_Combat/CombatGameMode.h
@@ -0,0 +1,20 @@
+// Copyright Epic Games, Inc. All Rights Reserved.
+
+#pragma once
+
+#include "CoreMinimal.h"
+#include "GameFramework/GameModeBase.h"
+#include "CombatGameMode.generated.h"
+
+/**
+ *  Simple GameMode for a third person combat game
+ */
+UCLASS(abstract)
+class ACombatGameMode : public AGameModeBase
+{
+	GENERATED_BODY()
+	
+public:
+
+	ACombatGameMode();
+};
diff --git a/Templates/TP_ThirdPerson/Source/TP_ThirdPerson/Variant_Combat/CombatLavaFloor.cpp b/Templates/TP_ThirdPerson/Source/TP_ThirdPerson/Variant_Combat/CombatLavaFloor.cpp
new file mode 100644
--- /dev/null
+++ b/Templates/TP_ThirdPerson/Source/TP_ThirdPerson/Variant_Combat/CombatLavaFloor.cpp
@@ -0,0 +1,27 @@
+// Copyright Epic Games, Inc. All Rights Reserved.
+
+
+#include "CombatLavaFloor.h"
+#include "CombatDamageable.h"
+#include "Components/StaticMeshComponent.h"
+
+ACombatLavaFloor::ACombatLavaFloor()
+{
+	PrimaryActorTick.bCanEverTick = false;
+
+	// create the mesh
+	RootComponent = Mesh = CreateDefaultSubobject<UStaticMeshComponent>(TEXT("Mesh"));
+
+	// bind the hit handler
+	Mesh->OnComponentHit.AddDynamic(this, &ACombatLavaFloor::OnFloorHit);
+}
+
+void ACombatLavaFloor::OnFloorHit(UPrimitiveComponent* HitComponent, AActor* OtherActor, UPrimitiveComponent* OtherComp, FVector NormalImpulse, const FHitResult& Hit)
+{
+	// check if the hit actor is damageable by casting to the interface
+	if (ICombatDamageable* Damageable = Cast<ICombatDamageable>(OtherActor))
+	{
+		// damage the actor
+		Damageable->ApplyDamage(Damage, this, Hit.ImpactPoint, FVector::ZeroVector);
+	}
+}
diff --git a/Templates/TP_ThirdPerson/Source/TP_ThirdPerson/Variant_Combat/CombatLavaFloor.h b/Templates/TP_ThirdPerson/Source/TP_ThirdPerson/Variant_Combat/CombatLavaFloor.h
new file mode 100644
--- /dev/null
+++ b/Templates/TP_ThirdPerson/Source/TP_ThirdPerson/Variant_Combat/CombatLavaFloor.h
@@ -0,0 +1,40 @@
+// Copyright Epic Games, Inc. All Rights Reserved.
+
+#pragma once
+
+#include "CoreMinimal.h"
+#include "GameFramework/Actor.h"
+#include "CombatLavaFloor.generated.h"
+
+class UStaticMeshComponent;
+class UPrimitiveComponent;
+
+/**
+ *  A basic actor that applies damage on contact through the ICombatDamageable interface. 
+ */
+UCLASS(abstract)
+class ACombatLavaFloor : public AActor
+{
+	GENERATED_BODY()
+	
+	/** Floor mesh */
+	UPROPERTY(VisibleAnywhere, BlueprintReadOnly, Category = Camera, meta = (AllowPrivateAccess = "true"))
+	UStaticMeshComponent* Mesh;
+
+protected:
+
+	/** Amount of damage to deal on contact */
+	UPROPERTY(EditAnywhere, Category="Damage")
+	float Damage = 10000.0f;
+
+public:	
+
+	/** Constructor */
+	ACombatLavaFloor();
+
+protected:
+
+	/** Blocking hit handler */
+	UFUNCTION()
+	void OnFloorHit(UPrimitiveComponent* HitComponent, AActor* OtherActor, UPrimitiveComponent* OtherComp, FVector NormalImpulse, const FHitResult& Hit);
+};
diff --git a/Templates/TP_ThirdPerson/Source/TP_ThirdPerson/Variant_Combat/CombatLifeBar.cpp b/Templates/TP_ThirdPerson/Source/TP_ThirdPerson/Variant_Combat/CombatLifeBar.cpp
new file mode 100644
--- /dev/null
+++ b/Templates/TP_ThirdPerson/Source/TP_ThirdPerson/Variant_Combat/CombatLifeBar.cpp
@@ -0,0 +1,5 @@
+// Copyright Epic Games, Inc. All Rights Reserved.
+
+
+#include "CombatLifeBar.h"
+
diff --git a/Templates/TP_ThirdPerson/Source/TP_ThirdPerson/Variant_Combat/CombatLifeBar.h b/Templates/TP_ThirdPerson/Source/TP_ThirdPerson/Variant_Combat/CombatLifeBar.h
new file mode 100644
--- /dev/null
+++ b/Templates/TP_ThirdPerson/Source/TP_ThirdPerson/Variant_Combat/CombatLifeBar.h
@@ -0,0 +1,26 @@
+// Copyright Epic Games, Inc. All Rights Reserved.
+
+#pragma once
+
+#include "CoreMinimal.h"
+#include "Blueprint/UserWidget.h"
+#include "CombatLifeBar.generated.h"
+
+/**
+ *  A basic life bar user widget.
+ */
+UCLASS(abstract)
+class UCombatLifeBar : public UUserWidget
+{
+	GENERATED_BODY()
+
+public:
+
+	/** Sets the life bar to the provided 0-1 percentage value*/
+	UFUNCTION(BlueprintImplementableEvent, Category="LifeBar")
+	void SetLifePercentage(float Percent);
+
+	// Sets the life bar fill color
+	UFUNCTION(BlueprintImplementableEvent, Category="LifeBar")
+	void SetBarColor(FLinearColor Color);
+};
diff --git a/Templates/TP_ThirdPerson/Source/TP_ThirdPerson/Variant_Combat/CombatPlayerController.cpp b/Templates/TP_ThirdPerson/Source/TP_ThirdPerson/Variant_Combat/CombatPlayerController.cpp
new file mode 100644
--- /dev/null
+++ b/Templates/TP_ThirdPerson/Source/TP_ThirdPerson/Variant_Combat/CombatPlayerController.cpp
@@ -0,0 +1,47 @@
+// Copyright Epic Games, Inc. All Rights Reserved.
+
+
+#include "Variant_Combat/CombatPlayerController.h"
+#include "EnhancedInputSubsystems.h"
+#include "InputMappingContext.h"
+#include "Kismet/GameplayStatics.h"
+#include "GameFramework/PlayerStart.h"
+#include "CombatCharacter.h"
+#include "Engine/LocalPlayer.h"
+#include "Engine/World.h"
+
+void ACombatPlayerController::SetupInputComponent()
+{
+	// add the input mapping context
+	if (UEnhancedInputLocalPlayerSubsystem* Subsystem = ULocalPlayer::GetSubsystem<UEnhancedInputLocalPlayerSubsystem>(GetLocalPlayer()))
+	{
+		for (UInputMappingContext* CurrentContext : DefaultMappingContexts)
+		{
+			Subsystem->AddMappingContext(CurrentContext, 0);
+		}
+	}
+}
+
+void ACombatPlayerController::OnPossess(APawn* InPawn)
+{
+	Super::OnPossess(InPawn);
+
+	// subscribe to the pawn's OnDestroyed delegate
+	InPawn->OnDestroyed.AddDynamic(this, &ACombatPlayerController::OnPawnDestroyed);
+}
+
+void ACombatPlayerController::SetRespawnTransform(const FTransform& NewRespawn)
+{
+	// save the new respawn transform
+	RespawnTransform = NewRespawn;
+}
+
+void ACombatPlayerController::OnPawnDestroyed(AActor* DestroyedActor)
+{
+	// spawn a new character at the respawn transform
+	if (ACombatCharacter* RespawnedCharacter = GetWorld()->SpawnActor<ACombatCharacter>(CharacterClass, RespawnTransform))
+	{
+		// possess the character
+		Possess(RespawnedCharacter);
+	}
+}
\ No newline at end of file
diff --git a/Templates/TP_ThirdPerson/Source/TP_ThirdPerson/Variant_Combat/CombatPlayerController.h b/Templates/TP_ThirdPerson/Source/TP_ThirdPerson/Variant_Combat/CombatPlayerController.h
new file mode 100644
--- /dev/null
+++ b/Templates/TP_ThirdPerson/Source/TP_ThirdPerson/Variant_Combat/CombatPlayerController.h
@@ -0,0 +1,54 @@
+// Copyright Epic Games, Inc. All Rights Reserved.
+
+#pragma once
+
+#include "CoreMinimal.h"
+#include "GameFramework/PlayerController.h"
+#include "CombatPlayerController.generated.h"
+
+class UInputMappingContext;
+class ACombatCharacter;
+
+/**
+ *  Simple Player Controller for a third person combat game
+ *  Manages input mappings
+ *  Respawns the player character at the checkpoint when it's destroyed
+ */
+UCLASS(abstract)
+class ACombatPlayerController : public APlayerController
+{
+	GENERATED_BODY()
+	
+protected:
+
+	/** Input mapping context for this player */
+	UPROPERTY(EditAnywhere, Category="Input")
+	TArray<UInputMappingContext*> DefaultMappingContexts;
+
+	/** Character class to respawn when the possessed pawn is destroyed */
+	UPROPERTY(EditAnywhere, Category="Respawn")
+	TSubclassOf<ACombatCharacter> CharacterClass;
+
+	/** Transform to respawn the character at. Can be set to create checkpoints */
+	FTransform RespawnTransform;
+
+protected:
+
+	/** Initialize input bindings */
+	virtual void SetupInputComponent() override;
+
+	/** Pawn initialization */
+	virtual void OnPossess(APawn* InPawn) override;
+
+public:
+
+	/** Updates the character respawn transform */
+	void SetRespawnTransform(const FTransform& NewRespawn);
+
+protected:
+
+	/** Called if the possessed pawn is destroyed */
+	UFUNCTION()
+	void OnPawnDestroyed(AActor* DestroyedActor);
+
+};
diff --git a/Templates/TP_ThirdPerson/Source/TP_ThirdPerson/Variant_Platforming/AnimNotify_EndDash.cpp b/Templates/TP_ThirdPerson/Source/TP_ThirdPerson/Variant_Platforming/AnimNotify_EndDash.cpp
new file mode 100644
--- /dev/null
+++ b/Templates/TP_ThirdPerson/Source/TP_ThirdPerson/Variant_Platforming/AnimNotify_EndDash.cpp
@@ -0,0 +1,21 @@
+// Copyright Epic Games, Inc. All Rights Reserved.
+
+
+#include "AnimNotify_EndDash.h"
+#include "PlatformingCharacter.h"
+#include "Components/SkeletalMeshComponent.h"
+
+void UAnimNotify_EndDash::Notify(USkeletalMeshComponent* MeshComp, UAnimSequenceBase* Animation, const FAnimNotifyEventReference& EventReference)
+{
+	// cast the owner to the attacker interface
+	if (APlatformingCharacter* PlatformingCharacter = Cast<APlatformingCharacter>(MeshComp->GetOwner()))
+	{
+		// tell the actor to end the dash
+		PlatformingCharacter->EndDash();
+	}
+}
+
+FString UAnimNotify_EndDash::GetNotifyName_Implementation() const
+{
+	return FString("End Dash");
+}
diff --git a/Templates/TP_ThirdPerson/Source/TP_ThirdPerson/Variant_Platforming/AnimNotify_EndDash.h b/Templates/TP_ThirdPerson/Source/TP_ThirdPerson/Variant_Platforming/AnimNotify_EndDash.h
new file mode 100644
--- /dev/null
+++ b/Templates/TP_ThirdPerson/Source/TP_ThirdPerson/Variant_Platforming/AnimNotify_EndDash.h
@@ -0,0 +1,24 @@
+// Copyright Epic Games, Inc. All Rights Reserved.
+
+#pragma once
+
+#include "CoreMinimal.h"
+#include "Animation/AnimNotifies/AnimNotify.h"
+#include "AnimNotify_EndDash.generated.h"
+
+/**
+ *  AnimNotify to finish the dash animation and restore player control
+ */
+UCLASS()
+class UAnimNotify_EndDash : public UAnimNotify
+{
+	GENERATED_BODY()
+	
+public:
+
+	/** Perform the Anim Notify */
+	virtual void Notify(USkeletalMeshComponent* MeshComp, UAnimSequenceBase* Animation, const FAnimNotifyEventReference& EventReference) override;
+
+	/** Get the notify name */
+	virtual FString GetNotifyName_Implementation() const override;
+};
diff --git a/Templates/TP_ThirdPerson/Source/TP_ThirdPerson/Variant_Platforming/PlatformingCharacter.cpp b/Templates/TP_ThirdPerson/Source/TP_ThirdPerson/Variant_Platforming/PlatformingCharacter.cpp
new file mode 100644
--- /dev/null
+++ b/Templates/TP_ThirdPerson/Source/TP_ThirdPerson/Variant_Platforming/PlatformingCharacter.cpp
@@ -0,0 +1,339 @@
+// Copyright Epic Games, Inc. All Rights Reserved.
+
+
+#include "PlatformingCharacter.h"
+
+#include "Components/CapsuleComponent.h"
+#include "GameFramework/CharacterMovementComponent.h"
+#include "GameFramework/SpringArmComponent.h"
+#include "Components/SkeletalMeshComponent.h"
+#include "Camera/CameraComponent.h"
+#include "EnhancedInputSubsystems.h"
+#include "EnhancedInputComponent.h"
+#include "TimerManager.h"
+#include "Engine/LocalPlayer.h"
+
+APlatformingCharacter::APlatformingCharacter()
+{
+ 	PrimaryActorTick.bCanEverTick = true;
+
+	// initialize the flags
+	bHasWallJumped = false;
+	bHasDoubleJumped = false;
+	bHasDashed = false;
+	bIsDashing = false;
+
+	// bind the attack montage ended delegate
+	OnDashMontageEnded.BindUObject(this, &APlatformingCharacter::DashMontageEnded);
+
+	// enable press and hold jump
+	JumpMaxHoldTime = 0.4f;
+
+	// set the jump max count to 2 so we can double jump
+	JumpMaxCount = 2;
+
+	// Set size for collision capsule
+	GetCapsuleComponent()->InitCapsuleSize(35.0f, 90.0f);
+
+	// don't rotate the mesh when the controller rotates
+	bUseControllerRotationYaw = false;
+	
+	// Configure character movement
+	GetCharacterMovement()->GravityScale = 2.5f;
+	GetCharacterMovement()->MaxAcceleration = 1500.0f;
+	GetCharacterMovement()->BrakingFrictionFactor = 1.0f;
+	GetCharacterMovement()->bUseSeparateBrakingFriction = true;
+
+	GetCharacterMovement()->GroundFriction = 4.0f;
+	GetCharacterMovement()->MaxWalkSpeed = 750.0f;
+	GetCharacterMovement()->MinAnalogWalkSpeed = 20.0f;
+	GetCharacterMovement()->BrakingDecelerationWalking = 2500.0f;
+	GetCharacterMovement()->PerchRadiusThreshold = 15.0f;
+
+	GetCharacterMovement()->JumpZVelocity = 350.0f;
+	GetCharacterMovement()->BrakingDecelerationFalling = 750.0f;
+	GetCharacterMovement()->AirControl = 1.0f;
+
+	GetCharacterMovement()->RotationRate = FRotator(0.0f, 500.0f, 0.0f);
+	GetCharacterMovement()->bOrientRotationToMovement = true;
+
+	GetCharacterMovement()->NavAgentProps.AgentRadius = 42.0;
+	GetCharacterMovement()->NavAgentProps.AgentHeight = 192.0;
+
+	// create the camera boom
+	CameraBoom = CreateDefaultSubobject<USpringArmComponent>(TEXT("CameraBoom"));
+	CameraBoom->SetupAttachment(RootComponent);
+
+	CameraBoom->TargetArmLength = 400.0f;
+	CameraBoom->bUsePawnControlRotation = true;
+	CameraBoom->bEnableCameraLag = true;
+	CameraBoom->CameraLagSpeed = 8.0f;
+	CameraBoom->bEnableCameraRotationLag = true;	
+	CameraBoom->CameraRotationLagSpeed = 8.0f;
+
+	// create the orbiting camera
+	FollowCamera = CreateDefaultSubobject<UCameraComponent>(TEXT("FollowCamera"));
+	FollowCamera->SetupAttachment(CameraBoom, USpringArmComponent::SocketName);
+	FollowCamera->bUsePawnControlRotation = false;
+}
+
+void APlatformingCharacter::Move(const FInputActionValue& Value)
+{
+	FVector2D MovementVector = Value.Get<FVector2D>();
+
+	// route the input
+	DoMove(MovementVector.X, MovementVector.Y);
+}
+
+void APlatformingCharacter::Look(const FInputActionValue& Value)
+{
+	FVector2D LookAxisVector = Value.Get<FVector2D>();
+
+	// route the input
+	DoLook(LookAxisVector.X, LookAxisVector.Y);
+}
+
+
+void APlatformingCharacter::Dash()
+{
+	// route the input
+	DoDash();
+}
+
+void APlatformingCharacter::MultiJump()
+{
+	// ignore jumps while dashing
+	if(bIsDashing)
+		return;
+
+	// are we already in the air?
+	if (GetCharacterMovement()->IsFalling())
+	{
+
+		// have we already wall jumped?
+		if (!bHasWallJumped)
+		{
+			// run a sphere sweep to check if we're in front of a wall
+			FHitResult OutHit;
+
+			const FVector TraceStart = GetActorLocation();
+			const FVector TraceEnd = TraceStart + (GetActorForwardVector() * WallJumpTraceDistance);
+			const FCollisionShape TraceShape = FCollisionShape::MakeSphere(WallJumpTraceRadius);
+
+			FCollisionQueryParams QueryParams;
+			QueryParams.AddIgnoredActor(this);
+
+			if (GetWorld()->SweepSingleByChannel(OutHit, TraceStart, TraceEnd, FQuat(), ECollisionChannel::ECC_Visibility, TraceShape, QueryParams))
+			{
+				// rotate the character to face away from the wall, so we're correctly oriented for the next wall jump
+				FRotator WallOrientation = OutHit.ImpactNormal.ToOrientationRotator();
+				WallOrientation.Pitch = 0.0f;
+				WallOrientation.Roll = 0.0f;
+
+				SetActorRotation(WallOrientation);
+
+				// apply a launch impulse to the character to perform the actual wall jump
+				const FVector WallJumpImpulse = (OutHit.ImpactNormal * WallJumpBounceImpulse) + (FVector::UpVector * WallJumpVerticalImpulse);
+
+				LaunchCharacter(WallJumpImpulse, true, true);
+
+				// enable the jump trail
+				SetJumpTrailState(true);
+
+				// raise the wall jump flag to prevent an immediate second wall jump
+				bHasWallJumped = true;
+
+				GetWorld()->GetTimerManager().SetTimer(WallJumpTimer, this, &APlatformingCharacter::ResetWallJump, DelayBetweenWallJumps, false);
+			}
+			// no wall jump, try a double jump next
+			else
+			{
+				// only double jump once while we're in the air
+				if (!bHasDoubleJumped)
+				{
+					bHasDoubleJumped = true;
+
+					// use the built-in CMC functionality to do the double jump
+					Jump();
+
+					// enable the jump trail
+					SetJumpTrailState(true);
+				}
+			}
+		}
+
+	}
+	else
+	{
+		// we're grounded so just do a regular jump
+		Jump();
+
+		// activate the jump trail
+		SetJumpTrailState(true);
+	}
+}
+
+void APlatformingCharacter::ResetWallJump()
+{
+	// reset the wall jump input lock
+	bHasWallJumped = false;
+}
+
+void APlatformingCharacter::DoMove(float Right, float Forward)
+{
+	if (GetController() != nullptr)
+	{
+		// momentarily disable movement inputs if we've just wall jumped
+		if (!bHasWallJumped)
+		{
+			// find out which way is forward
+			const FRotator Rotation = GetController()->GetControlRotation();
+			const FRotator YawRotation(0, Rotation.Yaw, 0);
+
+			// get forward vector
+			const FVector ForwardDirection = FRotationMatrix(YawRotation).GetUnitAxis(EAxis::X);
+
+			// get right vector 
+			const FVector RightDirection = FRotationMatrix(YawRotation).GetUnitAxis(EAxis::Y);
+
+			// add movement 
+			AddMovementInput(ForwardDirection, Forward);
+			AddMovementInput(RightDirection, Right);
+		}
+	}
+}
+
+void APlatformingCharacter::DoLook(float Yaw, float Pitch)
+{
+	if (GetController() != nullptr)
+	{
+		// add yaw and pitch input to controller
+		AddControllerYawInput(Yaw);
+		AddControllerPitchInput(Pitch);
+	}
+}
+
+void APlatformingCharacter::DoDash()
+{
+	// ignore the input if we've already dashed and have yet to reset
+	if (bHasDashed)
+		return;
+
+	// raise the dash flags
+	bIsDashing = true;
+	bHasDashed = true;
+
+	// disable gravity while dashing
+	GetCharacterMovement()->GravityScale = 0.0f;
+
+	// reset the character velocity so we don't carry momentum into the dash
+	GetCharacterMovement()->Velocity = FVector::ZeroVector;
+
+	// enable the jump trails
+	SetJumpTrailState(true);
+
+	// play the dash montage
+	if (UAnimInstance* AnimInstance = GetMesh()->GetAnimInstance())
+	{
+		const float MontageLength = AnimInstance->Montage_Play(DashMontage, 1.0f, EMontagePlayReturnType::MontageLength, 0.0f, true);
+
+		// has the montage played successfully?
+		if (MontageLength > 0.0f)
+		{
+			AnimInstance->Montage_SetEndDelegate(OnDashMontageEnded, DashMontage);
+		}
+	}
+}
+
+void APlatformingCharacter::DoJumpStart()
+{
+	// handle special jump cases
+	MultiJump();
+}
+
+void APlatformingCharacter::DoJumpEnd()
+{
+	// stop jumping
+	StopJumping();
+}
+
+void APlatformingCharacter::DashMontageEnded(UAnimMontage* Montage, bool bInterrupted)
+{
+	// if the montage was interrupted, end the dash
+	if (bInterrupted)
+	{
+		EndDash();
+	}
+}
+
+void APlatformingCharacter::EndDash()
+{
+	// restore gravity
+	GetCharacterMovement()->GravityScale = 2.5f;
+
+	// reset the dashing flag
+	bIsDashing = false;
+
+	// are we grounded after the dash?
+	if (GetCharacterMovement()->IsMovingOnGround())
+	{
+		// reset the dash usage flag, since we won't receive a landed event
+		bHasDashed = false;
+
+		// deactivate the jump trails
+		SetJumpTrailState(false);
+	}
+}
+
+bool APlatformingCharacter::HasDoubleJumped() const
+{
+	return bHasDoubleJumped;
+}
+
+bool APlatformingCharacter::HasWallJumped() const
+{
+	return bHasWallJumped;
+}
+
+void APlatformingCharacter::EndPlay(const EEndPlayReason::Type EndPlayReason)
+{
+	Super::EndPlay(EndPlayReason);
+
+	// clear the wall jump reset timer
+	GetWorld()->GetTimerManager().ClearTimer(WallJumpTimer);
+}
+
+void APlatformingCharacter::SetupPlayerInputComponent(UInputComponent* PlayerInputComponent)
+{
+	// Set up action bindings
+	if (UEnhancedInputComponent* EnhancedInputComponent = Cast<UEnhancedInputComponent>(PlayerInputComponent))
+	{
+
+		// Jumping
+		EnhancedInputComponent->BindAction(JumpAction, ETriggerEvent::Started, this, &APlatformingCharacter::DoJumpStart);
+		EnhancedInputComponent->BindAction(JumpAction, ETriggerEvent::Completed, this, &APlatformingCharacter::DoJumpEnd);
+
+		// Moving
+		EnhancedInputComponent->BindAction(MoveAction, ETriggerEvent::Triggered, this, &APlatformingCharacter::Move);
+		EnhancedInputComponent->BindAction(MouseLookAction, ETriggerEvent::Triggered, this, &APlatformingCharacter::Look);
+
+		// Looking
+		EnhancedInputComponent->BindAction(LookAction, ETriggerEvent::Triggered, this, &APlatformingCharacter::Look);
+
+		// Dashing
+		EnhancedInputComponent->BindAction(DashAction, ETriggerEvent::Triggered, this, &APlatformingCharacter::Dash);
+	}
+}
+
+void APlatformingCharacter::Landed(const FHitResult& Hit)
+{
+	Super::Landed(Hit);
+
+	// reset the double jump and dash flags
+	bHasDoubleJumped = false;
+	bHasDashed = false;
+
+	// deactivate the jump trail
+	SetJumpTrailState(false);
+}
+
diff --git a/Templates/TP_ThirdPerson/Source/TP_ThirdPerson/Variant_Platforming/PlatformingCharacter.h b/Templates/TP_ThirdPerson/Source/TP_ThirdPerson/Variant_Platforming/PlatformingCharacter.h
new file mode 100644
--- /dev/null
+++ b/Templates/TP_ThirdPerson/Source/TP_ThirdPerson/Variant_Platforming/PlatformingCharacter.h
@@ -0,0 +1,184 @@
+// Copyright Epic Games, Inc. All Rights Reserved.
+
+#pragma once
+
+#include "CoreMinimal.h"
+#include "GameFramework/Character.h"
+#include "Animation/AnimInstance.h"
+#include "PlatformingCharacter.generated.h"
+
+
+class USpringArmComponent;
+class UCameraComponent;
+class UInputAction;
+struct FInputActionValue;
+class UAnimMontage;
+
+/**
+ *  An enhanced Third Person Character with the following functionality:
+ *  - Platforming game character movement physics
+ *  - Press and Hold Jump
+ *  - Double Jump
+ *  - Wall Jump
+ *  - Dash
+ */
+UCLASS(abstract)
+class APlatformingCharacter : public ACharacter
+{
+	GENERATED_BODY()
+
+	/** Camera boom positioning the camera behind the character */
+	UPROPERTY(VisibleAnywhere, BlueprintReadOnly, Category = Camera, meta = (AllowPrivateAccess = "true"))
+	USpringArmComponent* CameraBoom;
+
+	/** Follow camera */
+	UPROPERTY(VisibleAnywhere, BlueprintReadOnly, Category = Camera, meta = (AllowPrivateAccess = "true"))
+	UCameraComponent* FollowCamera;
+	
+protected:
+
+	/** Jump Input Action */
+	UPROPERTY(EditAnywhere, BlueprintReadOnly, Category="Input")
+	UInputAction* JumpAction;
+
+	/** Move Input Action */
+	UPROPERTY(EditAnywhere, BlueprintReadOnly, Category="Input")
+	UInputAction* MoveAction;
+
+	/** Look Input Action */
+	UPROPERTY(EditAnywhere, BlueprintReadOnly, Category="Input")
+	UInputAction* LookAction;
+
+	/** Mouse Look Input Action */
+	UPROPERTY(EditAnywhere, BlueprintReadOnly, Category="Input")
+	UInputAction* MouseLookAction;
+
+	/** Dash Input Action */
+	UPROPERTY(EditAnywhere, BlueprintReadOnly, Category="Input")
+	UInputAction* DashAction;
+
+public:
+
+	/** Constructor */
+	APlatformingCharacter();
+
+protected:
+
+	/** Called for movement input */
+	void Move(const FInputActionValue& Value);
+
+	/** Called for looking input */
+	void Look(const FInputActionValue& Value);
+
+	/** Called for dash input */
+	void Dash();
+
+	/** Called for jump pressed to check for advanced multi-jump conditions */
+	void MultiJump();
+
+	/** Resets the wall jump input lock */
+	void ResetWallJump();
+
+public:
+
+	/** Handles move inputs from either controls or UI interfaces */
+	UFUNCTION(BlueprintCallable, Category="Input")
+	virtual void DoMove(float Right, float Forward);
+
+	/** Handles look inputs from either controls or UI interfaces */
+	UFUNCTION(BlueprintCallable, Category="Input")
+	virtual void DoLook(float Yaw, float Pitch);
+
+	/** Handles dash inputs from either controls or UI interfaces */
+	UFUNCTION(BlueprintCallable, Category="Input")
+	virtual void DoDash();
+
+	/** Handles jump pressed inputs from either controls or UI interfaces */
+	UFUNCTION(BlueprintCallable, Category="Input")
+	virtual void DoJumpStart();
+
+	/** Handles jump pressed inputs from either controls or UI interfaces */
+	UFUNCTION(BlueprintCallable, Category="Input")
+	virtual void DoJumpEnd();
+
+protected:
+
+	/** Called from a delegate when the dash montage ends */
+	void DashMontageEnded(UAnimMontage* Montage, bool bInterrupted);
+
+	/** Passes control to Blueprint to enable or disable jump trails */
+	UFUNCTION(BlueprintImplementableEvent, Category="Platforming")
+	void SetJumpTrailState(bool bEnabled);
+
+public:
+
+	/** Ends the dash state */
+	void EndDash();
+
+public:
+
+	/** Returns true if the character has just double jumped */
+	UFUNCTION(BlueprintPure, Category="Platforming")
+	bool HasDoubleJumped() const;
+
+	/** Returns true if the character has just wall jumped */
+	UFUNCTION(BlueprintPure, Category="Platforming")
+	bool HasWallJumped() const;
+
+public:	
+	
+	/** EndPlay cleanup */
+	virtual void EndPlay(const EEndPlayReason::Type EndPlayReason) override;
+
+	/** Sets up input action bindings */
+	virtual void SetupPlayerInputComponent(class UInputComponent* PlayerInputComponent) override;
+
+	/** Handle landings to reset dash and advanced jump state */
+	virtual void Landed(const FHitResult& Hit) override;
+
+protected:
+
+	/** movement state flag bits, packed into a uint8 for memory efficiency */
+	uint8 bHasWallJumped : 1;
+	uint8 bHasDoubleJumped : 1;
+	uint8 bHasDashed : 1;
+	uint8 bIsDashing : 1;
+
+	/** timer for wall jump input reset */
+	FTimerHandle WallJumpTimer;
+
+	/** Dash montage ended delegate */
+	FOnMontageEnded OnDashMontageEnded;
+
+	/** Distance to trace ahead of the character to look for walls to jump from */
+	UPROPERTY(EditAnywhere, Category="Wall Jump")
+	float WallJumpTraceDistance = 50.0f;
+
+	/** Radius of the wall jump sphere trace check */
+	UPROPERTY(EditAnywhere, Category="Wall Jump")
+	float WallJumpTraceRadius = 25.0f;
+
+	/** Impulse to apply away from the wall when wall jumping */
+	UPROPERTY(EditAnywhere, Category="Wall Jump")
+	float WallJumpBounceImpulse = 800.0f;
+
+	/** Vertical impulse to apply when wall jumping */
+	UPROPERTY(EditAnywhere, Category="Wall Jump")
+	float WallJumpVerticalImpulse = 900.0f;
+
+	/** Time to ignore jump inputs after a wall jump */
+	UPROPERTY(EditAnywhere, Category="Wall Jump")
+	float DelayBetweenWallJumps = 0.1f;
+
+	/** AnimMontage to use for the Dash action */
+	UPROPERTY(EditAnywhere, Category="Dash")
+	UAnimMontage* DashMontage;
+
+public:
+	/** Returns CameraBoom subobject **/
+	FORCEINLINE class USpringArmComponent* GetCameraBoom() const { return CameraBoom; }
+
+	/** Returns FollowCamera subobject **/
+	FORCEINLINE class UCameraComponent* GetFollowCamera() const { return FollowCamera; }
+
+};
diff --git a/Templates/TP_ThirdPerson/Source/TP_ThirdPerson/Variant_Platforming/PlatformingGameMode.cpp b/Templates/TP_ThirdPerson/Source/TP_ThirdPerson/Variant_Platforming/PlatformingGameMode.cpp
new file mode 100644
--- /dev/null
+++ b/Templates/TP_ThirdPerson/Source/TP_ThirdPerson/Variant_Platforming/PlatformingGameMode.cpp
@@ -0,0 +1,9 @@
+// Copyright Epic Games, Inc. All Rights Reserved.
+
+
+#include "Variant_Platforming/PlatformingGameMode.h"
+
+APlatformingGameMode::APlatformingGameMode()
+{
+	// stub
+}
\ No newline at end of file
diff --git a/Templates/TP_ThirdPerson/Source/TP_ThirdPerson/Variant_Platforming/PlatformingGameMode.h b/Templates/TP_ThirdPerson/Source/TP_ThirdPerson/Variant_Platforming/PlatformingGameMode.h
new file mode 100644
--- /dev/null
+++ b/Templates/TP_ThirdPerson/Source/TP_ThirdPerson/Variant_Platforming/PlatformingGameMode.h
@@ -0,0 +1,21 @@
+// Copyright Epic Games, Inc. All Rights Reserved.
+
+#pragma once
+
+#include "CoreMinimal.h"
+#include "GameFramework/GameModeBase.h"
+#include "PlatformingGameMode.generated.h"
+
+/**
+ *  Simple GameMode for a third person platforming game
+ */
+UCLASS()
+class APlatformingGameMode : public AGameModeBase
+{
+	GENERATED_BODY()
+	
+public:
+
+	/** Constructor */
+	APlatformingGameMode();
+};
diff --git a/Templates/TP_ThirdPerson/Source/TP_ThirdPerson/Variant_Platforming/PlatformingPlayerController.cpp b/Templates/TP_ThirdPerson/Source/TP_ThirdPerson/Variant_Platforming/PlatformingPlayerController.cpp
new file mode 100644
--- /dev/null
+++ b/Templates/TP_ThirdPerson/Source/TP_ThirdPerson/Variant_Platforming/PlatformingPlayerController.cpp
@@ -0,0 +1,50 @@
+// Copyright Epic Games, Inc. All Rights Reserved.
+
+
+#include "Variant_Platforming/PlatformingPlayerController.h"
+#include "EnhancedInputSubsystems.h"
+#include "InputMappingContext.h"
+#include "Kismet/GameplayStatics.h"
+#include "GameFramework/PlayerStart.h"
+#include "PlatformingCharacter.h"
+#include "Engine/LocalPlayer.h"
+#include "Engine/World.h"
+
+void APlatformingPlayerController::SetupInputComponent()
+{
+	// add the input mapping context
+	if (UEnhancedInputLocalPlayerSubsystem* Subsystem = ULocalPlayer::GetSubsystem<UEnhancedInputLocalPlayerSubsystem>(GetLocalPlayer()))
+	{
+		for (UInputMappingContext* CurrentContext : DefaultMappingContexts)
+		{
+			Subsystem->AddMappingContext(CurrentContext, 0);
+		}
+	}
+}
+
+void APlatformingPlayerController::OnPossess(APawn* InPawn)
+{
+	Super::OnPossess(InPawn);
+
+	// subscribe to the pawn's OnDestroyed delegate
+	InPawn->OnDestroyed.AddDynamic(this, &APlatformingPlayerController::OnPawnDestroyed);
+}
+
+void APlatformingPlayerController::OnPawnDestroyed(AActor* DestroyedActor)
+{
+	// find the player start
+	TArray<AActor*> ActorList;
+	UGameplayStatics::GetAllActorsOfClass(GetWorld(), APlayerStart::StaticClass(), ActorList);
+
+	if (ActorList.Num() > 0)
+	{
+		// spawn a character at the player start
+		const FTransform SpawnTransform = ActorList[0]->GetActorTransform();
+
+		if (APlatformingCharacter* RespawnedCharacter = GetWorld()->SpawnActor<APlatformingCharacter>(CharacterClass, SpawnTransform))
+		{
+			// possess the character
+			Possess(RespawnedCharacter);
+		}
+	}
+}
\ No newline at end of file
diff --git a/Templates/TP_ThirdPerson/Source/TP_ThirdPerson/Variant_Platforming/PlatformingPlayerController.h b/Templates/TP_ThirdPerson/Source/TP_ThirdPerson/Variant_Platforming/PlatformingPlayerController.h
new file mode 100644
--- /dev/null
+++ b/Templates/TP_ThirdPerson/Source/TP_ThirdPerson/Variant_Platforming/PlatformingPlayerController.h
@@ -0,0 +1,43 @@
+// Copyright Epic Games, Inc. All Rights Reserved.
+
+#pragma once
+
+#include "CoreMinimal.h"
+#include "GameFramework/PlayerController.h"
+#include "PlatformingPlayerController.generated.h"
+
+class UInputMappingContext;
+class APlatformingCharacter;
+
+/**
+ *  Simple Player Controller for a third person platforming game
+ *  Manages input mappings
+ *  Respawns the player character at the Player Start when it's destroyed
+ */
+UCLASS(abstract)
+class APlatformingPlayerController : public APlayerController
+{
+	GENERATED_BODY()
+	
+protected:
+
+	/** Input mapping context for this player */
+	UPROPERTY(EditAnywhere, Category="Input")
+	TArray<UInputMappingContext*> DefaultMappingContexts;
+
+	/** Character class to respawn when the possessed pawn is destroyed */
+	UPROPERTY(EditAnywhere, Category="Respawn")
+	TSubclassOf<APlatformingCharacter> CharacterClass;
+
+protected:
+
+	/** Initialize input bindings */
+	virtual void SetupInputComponent() override;
+
+	/** Pawn initialization */
+	virtual void OnPossess(APawn* InPawn) override;
+
+	/** Called if the possessed pawn is destroyed */
+	UFUNCTION()
+	void OnPawnDestroyed(AActor* DestroyedActor);
+};
diff --git a/Templates/TP_ThirdPerson/Source/TP_ThirdPerson/Variant_SideScrolling/AI/SideScrollingAIController.cpp b/Templates/TP_ThirdPerson/Source/TP_ThirdPerson/Variant_SideScrolling/AI/SideScrollingAIController.cpp
new file mode 100644
--- /dev/null
+++ b/Templates/TP_ThirdPerson/Source/TP_ThirdPerson/Variant_SideScrolling/AI/SideScrollingAIController.cpp
@@ -0,0 +1,19 @@
+// Copyright Epic Games, Inc. All Rights Reserved.
+
+
+#include "SideScrollingAIController.h"
+#include "GameplayStateTreeModule/Public/Components/StateTreeAIComponent.h"
+
+ASideScrollingAIController::ASideScrollingAIController()
+{
+	// create the StateTree AI Component
+	StateTreeAI = CreateDefaultSubobject<UStateTreeAIComponent>(TEXT("StateTreeAI"));
+	check(StateTreeAI);
+
+	// ensure we start the StateTree when we possess the pawn
+	bStartAILogicOnPossess = true;
+
+	// ensure we're attached to the possessed character.
+	// this is necessary for EnvQueries to work correctly
+	bAttachToPawn = true;
+}
diff --git a/Templates/TP_ThirdPerson/Source/TP_ThirdPerson/Variant_SideScrolling/AI/SideScrollingAIController.h b/Templates/TP_ThirdPerson/Source/TP_ThirdPerson/Variant_SideScrolling/AI/SideScrollingAIController.h
new file mode 100644
--- /dev/null
+++ b/Templates/TP_ThirdPerson/Source/TP_ThirdPerson/Variant_SideScrolling/AI/SideScrollingAIController.h
@@ -0,0 +1,27 @@
+// Copyright Epic Games, Inc. All Rights Reserved.
+
+#pragma once
+
+#include "CoreMinimal.h"
+#include "AIController.h"
+#include "SideScrollingAIController.generated.h"
+
+class UStateTreeAIComponent;
+
+/**
+ *  A basic AI Controller capable of running StateTree
+ */
+UCLASS(abstract)
+class ASideScrollingAIController : public AAIController
+{
+	GENERATED_BODY()
+	
+	/** StateTree Component */
+	UPROPERTY(VisibleAnywhere, BlueprintReadOnly, Category = "AI", meta = (AllowPrivateAccess = "true"))
+	UStateTreeAIComponent* StateTreeAI;
+
+public:
+
+	/** Constructor */
+	ASideScrollingAIController();
+};
diff --git a/Templates/TP_ThirdPerson/Source/TP_ThirdPerson/Variant_SideScrolling/AI/SideScrollingNPC.cpp b/Templates/TP_ThirdPerson/Source/TP_ThirdPerson/Variant_SideScrolling/AI/SideScrollingNPC.cpp
new file mode 100644
--- /dev/null
+++ b/Templates/TP_ThirdPerson/Source/TP_ThirdPerson/Variant_SideScrolling/AI/SideScrollingNPC.cpp
@@ -0,0 +1,52 @@
+// Copyright Epic Games, Inc. All Rights Reserved.
+
+
+#include "SideScrollingNPC.h"
+#include "GameFramework/CharacterMovementComponent.h"
+#include "TimerManager.h"
+
+ASideScrollingNPC::ASideScrollingNPC()
+{
+ 	PrimaryActorTick.bCanEverTick = true;
+
+	GetCharacterMovement()->MaxWalkSpeed = 150.0f;
+}
+
+void ASideScrollingNPC::EndPlay(EEndPlayReason::Type EndPlayReason)
+{
+	Super::EndPlay(EndPlayReason);
+
+	// clear the deactivation timer
+	GetWorld()->GetTimerManager().ClearTimer(DeactivationTimer);
+}
+
+void ASideScrollingNPC::Interaction(AActor* Interactor)
+{
+	// ignore if this NPC has already been deactivated
+	if (bDeactivated)
+	{
+		return;
+	}
+
+	// reset the deactivation flag
+	bDeactivated = true;
+
+	// stop character movement immediately
+	GetCharacterMovement()->StopMovementImmediately();
+
+	// launch the NPC away from the interactor
+	FVector LaunchVector = Interactor->GetActorForwardVector() * LaunchImpulse;
+	LaunchVector.Y = 0.0f;
+	LaunchVector.Z = LaunchVerticalImpulse;
+
+	LaunchCharacter(LaunchVector, true, true);
+
+	// set up a timer to schedule reactivation
+	GetWorld()->GetTimerManager().SetTimer(DeactivationTimer, this, &ASideScrollingNPC::ResetDeactivation, DeactivationTime, false);
+}
+
+void ASideScrollingNPC::ResetDeactivation()
+{
+	// reset the deactivation flag
+	bDeactivated = false;
+}
diff --git a/Templates/TP_ThirdPerson/Source/TP_ThirdPerson/Variant_SideScrolling/AI/SideScrollingNPC.h b/Templates/TP_ThirdPerson/Source/TP_ThirdPerson/Variant_SideScrolling/AI/SideScrollingNPC.h
new file mode 100644
--- /dev/null
+++ b/Templates/TP_ThirdPerson/Source/TP_ThirdPerson/Variant_SideScrolling/AI/SideScrollingNPC.h
@@ -0,0 +1,64 @@
+// Copyright Epic Games, Inc. All Rights Reserved.
+
+#pragma once
+
+#include "CoreMinimal.h"
+#include "GameFramework/Character.h"
+#include "SideScrollingInteractable.h"
+#include "SideScrollingNPC.generated.h"
+
+/**
+ *  Simple platforming NPC
+ *  Its behaviors will be dictated by a possessing AI Controller
+ *  It can be temporarily deactivated through Actor interactions
+ */
+UCLASS(abstract)
+class ASideScrollingNPC : public ACharacter, public ISideScrollingInteractable
+{
+	GENERATED_BODY()
+
+protected:
+
+	/** Horizontal impulse to apply to the NPC when it's interacted with */
+	UPROPERTY(EditAnywhere, Category="NPC", meta=(ClampMin=0, Units="cm/s"))
+	float LaunchImpulse = 500.0f;
+
+	/** Vertical impulse to apply to the NPC when it's interacted with */
+	UPROPERTY(EditAnywhere, Category="NPC", meta=(ClampMin=0, Units="cm/s"))
+	float LaunchVerticalImpulse = 500.0f;
+
+	/** Time that the NPC remains deactivated after being interacted with */
+	UPROPERTY(EditAnywhere, Category="NPC", meta=(ClampMin=0, Units="s"))
+	float DeactivationTime = 3.0f;
+
+public:
+
+	/** If true, this NPC is deactivated and will not be interacted with */
+	UPROPERTY(VisibleAnywhere, BlueprintReadOnly, Category="NPC")
+	bool bDeactivated = false;
+
+	/** Timer to reactivate the NPC */
+	FTimerHandle DeactivationTimer;
+
+public:
+
+	/** Constructor */
+	ASideScrollingNPC();
+
+public:
+
+	/** Cleanup */
+	virtual void EndPlay(EEndPlayReason::Type EndPlayReason) override;
+
+public:
+
+//	~begin IInteractable interface 
+
+	/** Performs an interaction triggered by another actor */
+	virtual void Interaction(AActor* Interactor) override;
+
+//	~end IInteractable interface
+
+	/** Reactivates the NPC */
+	void ResetDeactivation();
+};
diff --git a/Templates/TP_ThirdPerson/Source/TP_ThirdPerson/Variant_SideScrolling/AI/SideScrollingStateTreeUtility.cpp b/Templates/TP_ThirdPerson/Source/TP_ThirdPerson/Variant_SideScrolling/AI/SideScrollingStateTreeUtility.cpp
new file mode 100644
--- /dev/null
+++ b/Templates/TP_ThirdPerson/Source/TP_ThirdPerson/Variant_SideScrolling/AI/SideScrollingStateTreeUtility.cpp
@@ -0,0 +1,35 @@
+// Copyright Epic Games, Inc. All Rights Reserved.
+
+
+#include "SideScrollingStateTreeUtility.h"
+#include "StateTreeExecutionContext.h"
+#include "StateTreeExecutionTypes.h"
+#include "AIController.h"
+#include "Kismet/GameplayStatics.h"
+
+
+#define LOCTEXT_NAMESPACE "SideScrollingTemplate"
+
+EStateTreeRunStatus FStateTreeGetPlayerTask::Tick(FStateTreeExecutionContext& Context, const float DeltaTime) const
+{
+	// get the instance data
+	FInstanceDataType& InstanceData = Context.GetInstanceData(*this);
+
+	// set the player pawn as the target
+	InstanceData.TargetPlayer = UGameplayStatics::GetPlayerPawn(InstanceData.Controller.Get(), 0);
+
+	// are the NPC and target valid?
+	if (IsValid(InstanceData.TargetPlayer) && IsValid(InstanceData.NPC))
+	{
+		InstanceData.bValidTarget = FVector::Distance(InstanceData.NPC->GetActorLocation(), InstanceData.TargetPlayer->GetActorLocation()) < InstanceData.RangeMax;
+	}
+
+	return EStateTreeRunStatus::Running;
+}
+
+#if WITH_EDITOR
+FText FStateTreeGetPlayerTask::GetDescription(const FGuid& ID, FStateTreeDataView InstanceDataView, const IStateTreeBindingLookup& BindingLookup, EStateTreeNodeFormatting Formatting /*= EStateTreeNodeFormatting::Text*/) const
+{
+	return LOCTEXT("StateTreeTaskGetPlayer", "<b>Get Player</b>");
+}
+#endif // WITH_EDITOR
\ No newline at end of file
diff --git a/Templates/TP_ThirdPerson/Source/TP_ThirdPerson/Variant_SideScrolling/AI/SideScrollingStateTreeUtility.h b/Templates/TP_ThirdPerson/Source/TP_ThirdPerson/Variant_SideScrolling/AI/SideScrollingStateTreeUtility.h
new file mode 100644
--- /dev/null
+++ b/Templates/TP_ThirdPerson/Source/TP_ThirdPerson/Variant_SideScrolling/AI/SideScrollingStateTreeUtility.h
@@ -0,0 +1,59 @@
+// Copyright Epic Games, Inc. All Rights Reserved.
+
+#pragma once
+
+#include "CoreMinimal.h"
+#include "StateTreeTaskBase.h"
+
+#include "SideScrollingStateTreeUtility.generated.h"
+
+class AAIController;
+
+/**
+ *  Instance data for the FStateTreeGetPlayerTask task
+ */
+USTRUCT()
+struct FStateTreeGetPlayerInstanceData
+{
+	GENERATED_BODY()
+
+	/** NPC owning this task */
+	UPROPERTY(VisibleAnywhere, Category = Context)
+	TObjectPtr<APawn> NPC;
+
+	/** Holds the found player pawn */
+	UPROPERTY(VisibleAnywhere, Category = Context)
+	TObjectPtr<AAIController> Controller;
+
+	/** Holds the found player pawn */
+	UPROPERTY(VisibleAnywhere)
+	TObjectPtr<APawn> TargetPlayer;
+
+	/** Is the pawn close enough to be considered a valid target? */
+	UPROPERTY(VisibleAnywhere)
+	bool bValidTarget = false;
+
+	/** Max distance to be considered a valid target */
+	UPROPERTY(EditAnywhere, Category = Parameter, meta=(ClampMin = 0, Units = "cm"))
+	float RangeMax = 1000.0f;
+};
+
+/**
+ *  StateTree task to get the player-controlled character
+ */
+USTRUCT(meta=(DisplayName="Get Player", Category="Side Scrolling"))
+struct FStateTreeGetPlayerTask : public FStateTreeTaskCommonBase
+{
+	GENERATED_BODY()
+
+	/* Ensure we're using the correct instance data struct */
+	using FInstanceDataType = FStateTreeGetPlayerInstanceData;
+	virtual const UStruct* GetInstanceDataType() const override { return FInstanceDataType::StaticStruct(); }
+
+	/** Runs while the owning state is active */
+	virtual EStateTreeRunStatus Tick(FStateTreeExecutionContext& Context, const float DeltaTime) const override;
+
+#if WITH_EDITOR
+	virtual FText GetDescription(const FGuid& ID, FStateTreeDataView InstanceDataView, const IStateTreeBindingLookup& BindingLookup, EStateTreeNodeFormatting Formatting = EStateTreeNodeFormatting::Text) const override;
+#endif // WITH_EDITOR
+};
\ No newline at end of file
diff --git a/Templates/TP_ThirdPerson/Source/TP_ThirdPerson/Variant_SideScrolling/Gameplay/SideScrollingJumpPad.cpp b/Templates/TP_ThirdPerson/Source/TP_ThirdPerson/Variant_SideScrolling/Gameplay/SideScrollingJumpPad.cpp
new file mode 100644
--- /dev/null
+++ b/Templates/TP_ThirdPerson/Source/TP_ThirdPerson/Variant_SideScrolling/Gameplay/SideScrollingJumpPad.cpp
@@ -0,0 +1,46 @@
+// Copyright Epic Games, Inc. All Rights Reserved.
+
+
+#include "SideScrollingJumpPad.h"
+#include "Components/BoxComponent.h"
+#include "GameFramework/Character.h"
+#include "GameFramework/CharacterMovementComponent.h"
+#include "Components/SceneComponent.h"
+
+ASideScrollingJumpPad::ASideScrollingJumpPad()
+{
+	PrimaryActorTick.bCanEverTick = false;
+
+	// create the root comp
+	RootComponent = CreateDefaultSubobject<USceneComponent>(TEXT("Root"));
+
+	// create the bounding box
+	Box = CreateEditorOnlyDefaultSubobject<UBoxComponent>(TEXT("Box"));
+	Box->SetupAttachment(RootComponent);
+
+	// configure the bounding box
+	Box->SetBoxExtent(FVector(115.0f, 90.0f, 20.0f));
+	Box->SetRelativeLocation(FVector(0.0f, 0.0f, 16.0f));
+
+	Box->SetCollisionObjectType(ECC_WorldDynamic);
+	Box->SetCollisionEnabled(ECollisionEnabled::QueryOnly);
+	Box->SetCollisionResponseToAllChannels(ECR_Ignore);
+	Box->SetCollisionResponseToChannel(ECC_Pawn, ECR_Overlap);
+
+	// add the overlap handler
+	OnActorBeginOverlap.AddDynamic(this, &ASideScrollingJumpPad::BeginOverlap);
+}
+
+void ASideScrollingJumpPad::BeginOverlap(AActor* OverlappedActor, AActor* OtherActor)
+{
+	// were we overlapped by a character?
+	if (ACharacter* OverlappingCharacter = Cast<ACharacter>(OtherActor))
+	{
+		// force the character to jump
+		OverlappingCharacter->Jump();
+
+		// launch the character to override its vertical velocity
+		FVector LaunchVelocity = FVector::UpVector * ZStrength;
+		OverlappingCharacter->LaunchCharacter(LaunchVelocity, false, true);
+	}
+}
diff --git a/Templates/TP_ThirdPerson/Source/TP_ThirdPerson/Variant_SideScrolling/Gameplay/SideScrollingJumpPad.h b/Templates/TP_ThirdPerson/Source/TP_ThirdPerson/Variant_SideScrolling/Gameplay/SideScrollingJumpPad.h
new file mode 100644
--- /dev/null
+++ b/Templates/TP_ThirdPerson/Source/TP_ThirdPerson/Variant_SideScrolling/Gameplay/SideScrollingJumpPad.h
@@ -0,0 +1,39 @@
+// Copyright Epic Games, Inc. All Rights Reserved.
+
+#pragma once
+
+#include "CoreMinimal.h"
+#include "GameFramework/Actor.h"
+#include "SideScrollingJumpPad.generated.h"
+
+class UBoxComponent;
+
+/**
+ *  A simple jump pad that launches characters into the air
+ */
+UCLASS(abstract)
+class ASideScrollingJumpPad : public AActor
+{
+	GENERATED_BODY()
+	
+	/** Jump pad bounding box */
+	UPROPERTY(VisibleAnywhere, BlueprintReadOnly, Category ="Components", meta = (AllowPrivateAccess = "true"))
+	UBoxComponent* Box;
+
+protected:
+
+	/** Vertical velocity to set the character to when they use the jump pad*/
+	UPROPERTY(EditAnywhere, Category="Jump Pad", meta=(ClampMin=0, ClampMax=10000, Units="cm/s"))
+	float ZStrength = 1000.0f;
+
+public:	
+
+	/** Constructor */
+	ASideScrollingJumpPad();
+
+protected:
+
+	UFUNCTION()
+	void BeginOverlap(AActor* OverlappedActor, AActor* OtherActor);
+
+};
diff --git a/Templates/TP_ThirdPerson/Source/TP_ThirdPerson/Variant_SideScrolling/Gameplay/SideScrollingMovingPlatform.cpp b/Templates/TP_ThirdPerson/Source/TP_ThirdPerson/Variant_SideScrolling/Gameplay/SideScrollingMovingPlatform.cpp
new file mode 100644
--- /dev/null
+++ b/Templates/TP_ThirdPerson/Source/TP_ThirdPerson/Variant_SideScrolling/Gameplay/SideScrollingMovingPlatform.cpp
@@ -0,0 +1,40 @@
+// Copyright Epic Games, Inc. All Rights Reserved.
+
+
+#include "SideScrollingMovingPlatform.h"
+#include "Components/SceneComponent.h"
+
+ASideScrollingMovingPlatform::ASideScrollingMovingPlatform()
+{
+	PrimaryActorTick.bCanEverTick = false;
+
+	// create the root comp
+	RootComponent = CreateDefaultSubobject<USceneComponent>(TEXT("Root"));
+}
+
+void ASideScrollingMovingPlatform::Interaction(AActor* Interactor)
+{
+	// ignore interactions if we're already moving
+	if (bMoving)
+	{
+		return;
+	}
+
+	// raise the movement flag
+	bMoving = true;
+
+	// pass control to BP for the actual movement
+	BP_MoveToTarget();
+}
+
+void ASideScrollingMovingPlatform::ResetInteraction()
+{
+	// ignore if this is a one-shot platform
+	if (bOneShot)
+	{
+		return;
+	}
+
+	// reset the movement flag
+	bMoving = false;
+}
diff --git a/Templates/TP_ThirdPerson/Source/TP_ThirdPerson/Variant_SideScrolling/Gameplay/SideScrollingMovingPlatform.h b/Templates/TP_ThirdPerson/Source/TP_ThirdPerson/Variant_SideScrolling/Gameplay/SideScrollingMovingPlatform.h
new file mode 100644
--- /dev/null
+++ b/Templates/TP_ThirdPerson/Source/TP_ThirdPerson/Variant_SideScrolling/Gameplay/SideScrollingMovingPlatform.h
@@ -0,0 +1,60 @@
+// Copyright Epic Games, Inc. All Rights Reserved.
+
+#pragma once
+
+#include "CoreMinimal.h"
+#include "GameFramework/Actor.h"
+#include "SideScrollingInteractable.h"
+#include "SideScrollingMovingPlatform.generated.h"
+
+/**
+ *  Simple moving platform that can be triggered through interactions by other actors.
+ *  The actual movement is performed by Blueprint code through latent execution nodes.
+ */
+UCLASS(abstract)
+class ASideScrollingMovingPlatform : public AActor, public ISideScrollingInteractable
+{
+	GENERATED_BODY()
+	
+public:	
+	
+	/** Constructor */
+	ASideScrollingMovingPlatform();
+
+protected:
+
+	/** If this is true, the platform is mid-movement and will ignore further interactions */
+	bool bMoving = false;
+
+	/** Destination of the platform in world space */
+	UPROPERTY(EditAnywhere, BlueprintReadOnly, Category="Moving Platform")
+	FVector PlatformTarget;
+
+	/** Time for the platform to move to the destination */
+	UPROPERTY(EditAnywhere, BlueprintReadOnly, Category="Moving Platform", meta=(ClampMin=0, Units="s"))
+	float MoveDuration = 5.0f;
+
+	/** If this is true, the platform will only move once. */
+	UPROPERTY(EditAnywhere, Category="Moving Platform")
+	bool bOneShot = false;
+
+public:
+
+// ~begin IInteractable interface 
+
+	/** Performs an interaction triggered by another actor */
+	virtual void Interaction(AActor* Interactor) override;
+
+// ~end IInteractable interface
+
+	/** Resets the interaction state. Must be called from BP code to reset the platform */
+	UFUNCTION(BlueprintCallable, Category="Moving Platform")
+	virtual void ResetInteraction();
+
+protected:
+
+	/** Allows Blueprint code to do the actual platform movement */
+	UFUNCTION(BlueprintImplementableEvent, BlueprintCallable, Category="Moving Platform", meta=(DisplayName="Move to Target"))
+	void BP_MoveToTarget();
+
+};
diff --git a/Templates/TP_ThirdPerson/Source/TP_ThirdPerson/Variant_SideScrolling/Gameplay/SideScrollingPickup.cpp b/Templates/TP_ThirdPerson/Source/TP_ThirdPerson/Variant_SideScrolling/Gameplay/SideScrollingPickup.cpp
new file mode 100644
--- /dev/null
+++ b/Templates/TP_ThirdPerson/Source/TP_ThirdPerson/Variant_SideScrolling/Gameplay/SideScrollingPickup.cpp
@@ -0,0 +1,55 @@
+// Copyright Epic Games, Inc. All Rights Reserved.
+
+
+#include "SideScrollingPickup.h"
+#include "GameFramework/Character.h"
+#include "SideScrollingGameMode.h"
+#include "Components/SphereComponent.h"
+#include "Components/SceneComponent.h"
+#include "Engine/World.h"
+
+ASideScrollingPickup::ASideScrollingPickup()
+{
+	PrimaryActorTick.bCanEverTick = false;
+
+	// create the root comp
+	RootComponent = CreateDefaultSubobject<USceneComponent>(TEXT("Root"));
+
+	// create the bounding sphere
+	Sphere = CreateDefaultSubobject<USphereComponent>(TEXT("Collision"));
+	Sphere->SetupAttachment(RootComponent);
+
+	Sphere->SetSphereRadius(100.0f);
+
+	Sphere->SetCollisionObjectType(ECC_WorldDynamic);
+	Sphere->SetCollisionEnabled(ECollisionEnabled::QueryOnly);
+	Sphere->SetCollisionResponseToAllChannels(ECR_Ignore);
+	Sphere->SetCollisionResponseToChannel(ECC_Pawn, ECR_Overlap);
+
+	// add the overlap handler
+	OnActorBeginOverlap.AddDynamic(this, &ASideScrollingPickup::BeginOverlap);
+}
+
+void ASideScrollingPickup::BeginOverlap(AActor* OverlappedActor, AActor* OtherActor)
+{
+	// have we collided against a character?
+	if (ACharacter* OverlappedCharacter = Cast<ACharacter>(OtherActor))
+	{
+		// is this the player character?
+		if (OverlappedCharacter->IsPlayerControlled())
+		{
+			// get the game mode
+			if (ASideScrollingGameMode* GM = Cast<ASideScrollingGameMode>(GetWorld()->GetAuthGameMode()))
+			{
+				// tell the game mode to process a pickup
+				GM->ProcessPickup();
+
+				// disable collision so we don't get picked up again
+				SetActorEnableCollision(false);
+
+				// Call the BP handler. It will be responsible for destroying the pickup
+				BP_OnPickedUp();
+			}
+		}
+	}
+}
\ No newline at end of file
diff --git a/Templates/TP_ThirdPerson/Source/TP_ThirdPerson/Variant_SideScrolling/Gameplay/SideScrollingPickup.h b/Templates/TP_ThirdPerson/Source/TP_ThirdPerson/Variant_SideScrolling/Gameplay/SideScrollingPickup.h
new file mode 100644
--- /dev/null
+++ b/Templates/TP_ThirdPerson/Source/TP_ThirdPerson/Variant_SideScrolling/Gameplay/SideScrollingPickup.h
@@ -0,0 +1,38 @@
+// Copyright Epic Games, Inc. All Rights Reserved.
+
+#pragma once
+
+#include "CoreMinimal.h"
+#include "GameFramework/Actor.h"
+#include "SideScrollingPickup.generated.h"
+
+class USphereComponent;
+
+/**
+ *  A simple side scrolling game pickup
+ *  Increments a counter on the GameMode
+ */
+UCLASS(abstract)
+class ASideScrollingPickup : public AActor
+{
+	GENERATED_BODY()
+	
+	/** Pickup bounding sphere */
+	UPROPERTY(VisibleAnywhere, BlueprintReadOnly, Category ="Components", meta = (AllowPrivateAccess = "true"))
+	USphereComponent* Sphere;
+
+public:
+
+	/** Constructor */
+	ASideScrollingPickup();
+
+protected:
+
+	/** Handles pickup collision */
+	UFUNCTION()
+	void BeginOverlap(AActor* OverlappedActor, AActor* OtherActor);
+
+	/** Passes control to BP to play effects on pickup */
+	UFUNCTION(BlueprintImplementableEvent, Category="Pickup", meta=(DisplayName = "On Picked Up"))
+	void BP_OnPickedUp();
+};
diff --git a/Templates/TP_ThirdPerson/Source/TP_ThirdPerson/Variant_SideScrolling/Gameplay/SideScrollingSoftPlatform.cpp b/Templates/TP_ThirdPerson/Source/TP_ThirdPerson/Variant_SideScrolling/Gameplay/SideScrollingSoftPlatform.cpp
new file mode 100644
--- /dev/null
+++ b/Templates/TP_ThirdPerson/Source/TP_ThirdPerson/Variant_SideScrolling/Gameplay/SideScrollingSoftPlatform.cpp
@@ -0,0 +1,59 @@
+// Copyright Epic Games, Inc. All Rights Reserved.
+
+
+#include "SideScrollingSoftPlatform.h"
+#include "Components/SceneComponent.h"
+#include "Components/StaticMeshComponent.h"
+#include "Components/BoxComponent.h"
+#include "SideScrollingCharacter.h"
+
+ASideScrollingSoftPlatform::ASideScrollingSoftPlatform()
+{
+ 	PrimaryActorTick.bCanEverTick = true;
+
+	// create the root component
+	RootComponent = Root = CreateDefaultSubobject<USceneComponent>(TEXT("Root"));
+
+	// create the mesh
+	Mesh = CreateDefaultSubobject<UStaticMeshComponent>(TEXT("Mesh"));
+	Mesh->SetupAttachment(Root);
+
+	Mesh->SetCollisionEnabled(ECollisionEnabled::QueryAndPhysics);
+	Mesh->SetCollisionObjectType(ECC_WorldStatic);
+	Mesh->SetCollisionResponseToAllChannels(ECR_Block);
+
+	// create the collision check box
+	CollisionCheckBox = CreateDefaultSubobject<UBoxComponent>(TEXT("Collision Check Box"));
+	CollisionCheckBox->SetupAttachment(Mesh);
+
+	CollisionCheckBox->SetRelativeLocation(FVector(0.0f, 0.0f, -40.0f));
+	CollisionCheckBox->SetCollisionEnabled(ECollisionEnabled::QueryOnly);
+	CollisionCheckBox->SetCollisionObjectType(ECC_WorldDynamic);
+	CollisionCheckBox->SetCollisionResponseToAllChannels(ECR_Ignore);
+	CollisionCheckBox->SetCollisionResponseToChannel(ECC_Pawn, ECR_Overlap);
+
+	// subscribe to the overlap events
+	CollisionCheckBox->OnComponentBeginOverlap.AddDynamic(this, &ASideScrollingSoftPlatform::OnSoftCollisionOverlap);
+}
+
+void ASideScrollingSoftPlatform::OnSoftCollisionOverlap(UPrimitiveComponent* OverlappedComponent, AActor* OtherActor, UPrimitiveComponent* OtherComp, int32 OtherBodyIndex, bool bFromSweep, const FHitResult& SweepResult)
+{
+	// have we overlapped a character?
+	if (ASideScrollingCharacter* Char = Cast<ASideScrollingCharacter>(OtherActor))
+	{
+		// disable the soft collision channel
+		Char->SetSoftCollision(true);
+	}
+}
+
+void ASideScrollingSoftPlatform::NotifyActorEndOverlap(AActor* OtherActor)
+{
+	Super::NotifyActorEndOverlap(OtherActor);
+
+	// have we overlapped a character?
+	if (ASideScrollingCharacter* Char = Cast<ASideScrollingCharacter>(OtherActor))
+	{
+		// enable the soft collision channel
+		Char->SetSoftCollision(false);
+	}
+}
diff --git a/Templates/TP_ThirdPerson/Source/TP_ThirdPerson/Variant_SideScrolling/Gameplay/SideScrollingSoftPlatform.h b/Templates/TP_ThirdPerson/Source/TP_ThirdPerson/Variant_SideScrolling/Gameplay/SideScrollingSoftPlatform.h
new file mode 100644
--- /dev/null
+++ b/Templates/TP_ThirdPerson/Source/TP_ThirdPerson/Variant_SideScrolling/Gameplay/SideScrollingSoftPlatform.h
@@ -0,0 +1,46 @@
+// Copyright Epic Games, Inc. All Rights Reserved.
+
+#pragma once
+
+#include "CoreMinimal.h"
+#include "GameFramework/Actor.h"
+#include "SideScrollingSoftPlatform.generated.h"
+
+class USceneComponent;
+class UStaticMeshComponent;
+class UBoxComponent;
+
+/**
+ *  A side scrolling game platform that the character can jump or drop through.
+ */
+UCLASS(abstract)
+class ASideScrollingSoftPlatform : public AActor
+{
+	GENERATED_BODY()
+	
+	/** Root component */
+	UPROPERTY(VisibleAnywhere, BlueprintReadOnly, Category ="Components", meta = (AllowPrivateAccess = "true"))
+	USceneComponent* Root;
+
+	/** Platform mesh. The part we collide against and see */
+	UPROPERTY(VisibleAnywhere, BlueprintReadOnly, Category ="Components", meta = (AllowPrivateAccess = "true"))
+	UStaticMeshComponent* Mesh;
+
+	/** Collision volume that toggles soft collision on the character when they're below the platform. */
+	UPROPERTY(VisibleAnywhere, BlueprintReadOnly, Category ="Components", meta = (AllowPrivateAccess = "true"))
+	UBoxComponent* CollisionCheckBox;
+
+public:	
+	
+	/** Constructor */
+	ASideScrollingSoftPlatform();
+
+protected:
+
+	/** Handles soft collision check box overlaps */
+	UFUNCTION()
+	void OnSoftCollisionOverlap(UPrimitiveComponent* OverlappedComponent, AActor* OtherActor, UPrimitiveComponent* OtherComp, int32 OtherBodyIndex, bool bFromSweep, const FHitResult& SweepResult);
+
+	/** Restores soft collision state when overlap ends */
+	virtual void NotifyActorEndOverlap(AActor* OtherActor) override;
+};
diff --git a/Templates/TP_ThirdPerson/Source/TP_ThirdPerson/Variant_SideScrolling/SideScrollingCPP.Build.cs b/Templates/TP_ThirdPerson/Source/TP_ThirdPerson/Variant_SideScrolling/SideScrollingCPP.Build.cs
new file mode 100644
--- /dev/null
+++ b/Templates/TP_ThirdPerson/Source/TP_ThirdPerson/Variant_SideScrolling/SideScrollingCPP.Build.cs
@@ -0,0 +1,29 @@
+// Copyright Epic Games, Inc. All Rights Reserved.
+
+using UnrealBuildTool;
+
+public class SideScrollingCPP : ModuleRules
+{
+	public SideScrollingCPP(ReadOnlyTargetRules Target) : base(Target)
+	{
+		PCHUsage = PCHUsageMode.UseExplicitOrSharedPCHs;
+	
+		PublicDependencyModuleNames.AddRange(new string[] {
+			"Core",
+			"CoreUObject",
+			"Engine",
+			"InputCore",
+			"EnhancedInput",
+			"StateTreeModule",
+			"GameplayStateTreeModule"
+		});
+
+		PrivateDependencyModuleNames.AddRange(new string[] {  });
+
+		PublicIncludePaths.AddRange(new string[] {
+			"SideScrollingCPP",
+			"SideScrollingCPP/NPC",
+			"SideScrollingCPP/Gameplay"
+		});
+	}
+}
diff --git a/Templates/TP_ThirdPerson/Source/TP_ThirdPerson/Variant_SideScrolling/SideScrollingCameraManager.cpp b/Templates/TP_ThirdPerson/Source/TP_ThirdPerson/Variant_SideScrolling/SideScrollingCameraManager.cpp
new file mode 100644
--- /dev/null
+++ b/Templates/TP_ThirdPerson/Source/TP_ThirdPerson/Variant_SideScrolling/SideScrollingCameraManager.cpp
@@ -0,0 +1,105 @@
+// Copyright Epic Games, Inc. All Rights Reserved.
+
+
+#include "SideScrollingCameraManager.h"
+#include "GameFramework/Pawn.h"
+#include "Engine/HitResult.h"
+#include "CollisionQueryParams.h"
+#include "Engine/World.h"
+
+void ASideScrollingCameraManager::UpdateViewTarget(FTViewTarget& OutVT, float DeltaTime)
+{
+	// ensure the view target is a pawn
+	APawn* TargetPawn = Cast<APawn>(OutVT.Target);
+
+	// is our target valid?
+	if (IsValid(TargetPawn))
+	{
+		// set the view target FOV and rotation
+		OutVT.POV.Rotation = FRotator(0.0f, -90.0f, 0.0f);
+		OutVT.POV.FOV = 65.0f;
+
+		// cache the current location
+		FVector CurrentActorLocation = OutVT.Target->GetActorLocation();
+
+		// copy the current camera location
+		FVector CurrentCameraLocation = GetCameraLocation();
+
+		// calculate the "zoom distance" - in reality the distance we want to keep to the target
+		float CurrentY = CurrentZoom + CurrentActorLocation.Y;
+
+		// do first-time setup
+		if (bSetup)
+		{
+			// lower the setup flag
+			bSetup = false;
+
+			// initialize the camera viewpoint and return
+			OutVT.POV.Location.X = CurrentActorLocation.X;
+			OutVT.POV.Location.Y = CurrentY;
+			OutVT.POV.Location.Z = CurrentActorLocation.Z + CameraZOffset;
+
+			// save the current camera height
+			CurrentZ = OutVT.POV.Location.Z;
+
+			// skip the rest of the calculations
+			return;
+		}
+
+		// check if the camera needs to update its height
+		bool bZUpdate = false;
+
+		// is the character moving vertically?
+		if (FMath::IsNearlyZero(TargetPawn->GetVelocity().Z))
+		{
+			// determine if we need to do a height update
+			bZUpdate = FMath::IsNearlyEqual(CurrentZ, CurrentCameraLocation.Z, 25.0f);
+
+		} else {
+
+			// run a trace below the character to determine if we need to do a height update
+			FHitResult OutHit;
+
+			const FVector End = CurrentActorLocation + FVector(0.0f, 0.0f, -1000.0f);
+
+			FCollisionQueryParams QueryParams;
+			QueryParams.AddIgnoredActor(TargetPawn);
+
+			// only update height if we're not about to hit ground
+			bZUpdate = !GetWorld()->LineTraceSingleByChannel(OutHit, CurrentActorLocation, End, ECC_Visibility, QueryParams);
+
+		}
+
+		// do we need to do a height update?
+		if (bZUpdate)
+		{
+
+			// set the height goal from the actor location
+			CurrentZ = CurrentActorLocation.Z;
+
+		} else {
+
+			// are we close enough to the target height?
+			if (FMath::IsNearlyEqual(CurrentZ, CurrentActorLocation.Z, 100.0f))
+			{
+				// set the height goal from the actor location
+				CurrentZ = CurrentActorLocation.Z;
+
+			} else {
+
+				// blend the height towards the actor location
+				CurrentZ = FMath::FInterpTo(CurrentZ, CurrentActorLocation.Z, DeltaTime, 2.0f);
+				
+			}
+
+		}
+
+		// clamp the X axis to the min and max camera bounds
+		float CurrentX = FMath::Clamp(CurrentActorLocation.X, CameraXMinBounds, CameraXMaxBounds);
+
+		// blend towards the new camera location and update the output
+		FVector TargetCameraLocation(CurrentX, CurrentY, CurrentZ);
+
+		OutVT.POV.Location = FMath::VInterpTo(CurrentCameraLocation, TargetCameraLocation, DeltaTime, 2.0f);
+	}
+}
\ No newline at end of file
diff --git a/Templates/TP_ThirdPerson/Source/TP_ThirdPerson/Variant_SideScrolling/SideScrollingCameraManager.h b/Templates/TP_ThirdPerson/Source/TP_ThirdPerson/Variant_SideScrolling/SideScrollingCameraManager.h
new file mode 100644
--- /dev/null
+++ b/Templates/TP_ThirdPerson/Source/TP_ThirdPerson/Variant_SideScrolling/SideScrollingCameraManager.h
@@ -0,0 +1,47 @@
+// Copyright Epic Games, Inc. All Rights Reserved.
+
+#pragma once
+
+#include "CoreMinimal.h"
+#include "Camera/PlayerCameraManager.h"
+#include "SideScrollingCameraManager.generated.h"
+
+/**
+ *  Simple side scrolling camera with smooth scrolling and horizontal bounds
+ */
+UCLASS()
+class ASideScrollingCameraManager : public APlayerCameraManager
+{
+	GENERATED_BODY()
+	
+public:
+
+	/** Overrides the default camera view target calculation */
+	virtual void UpdateViewTarget(FTViewTarget& OutVT, float DeltaTime) override;
+
+public:
+
+	/** How close we want to stay to the view target */
+	UPROPERTY(EditAnywhere, BlueprintReadWrite, Category="Side Scrolling Camera", meta=(ClampMin=0, ClampMax=10000, Units="cm"))
+	float CurrentZoom = 1000.0f;
+
+	/** How far above the target do we want the camera to focus */
+	UPROPERTY(EditAnywhere, BlueprintReadWrite, Category="Side Scrolling Camera", meta=(ClampMin=0, ClampMax=10000, Units="cm"))
+	float CameraZOffset = 100.0f;
+
+	/** Minimum camera scrolling bounds in world space */
+	UPROPERTY(EditAnywhere, BlueprintReadWrite, Category="Side Scrolling Camera", meta=(ClampMin=0, ClampMax=10000, Units="cm"))
+	float CameraXMinBounds = -400.0f;
+
+	/** Maximum camera scrolling bounds in world space */
+	UPROPERTY(EditAnywhere, BlueprintReadWrite, Category="Side Scrolling Camera", meta=(ClampMin=0, ClampMax=10000, Units="cm"))
+	float CameraXMaxBounds = 10000.0f;
+
+protected:
+
+	/** Last cached camera vertical location. The camera only adjusts its height if necessary. */
+	float CurrentZ = 0.0f;
+
+	/** First-time update camera setup flag */
+	bool bSetup = true;
+};
diff --git a/Templates/TP_ThirdPerson/Source/TP_ThirdPerson/Variant_SideScrolling/SideScrollingCharacter.cpp b/Templates/TP_ThirdPerson/Source/TP_ThirdPerson/Variant_SideScrolling/SideScrollingCharacter.cpp
new file mode 100644
--- /dev/null
+++ b/Templates/TP_ThirdPerson/Source/TP_ThirdPerson/Variant_SideScrolling/SideScrollingCharacter.cpp
@@ -0,0 +1,324 @@
+// Copyright Epic Games, Inc. All Rights Reserved.
+
+
+#include "SideScrollingCharacter.h"
+#include "GameFramework/CharacterMovementComponent.h"
+#include "Components/CapsuleComponent.h"
+#include "Camera/CameraComponent.h"
+#include "Components/InputComponent.h"
+#include "InputActionValue.h"
+#include "EnhancedInputComponent.h"
+#include "InputAction.h"
+#include "Engine/World.h"
+#include "SideScrollingInteractable.h"
+#include "Kismet/KismetMathLibrary.h"
+#include "TimerManager.h"
+
+ASideScrollingCharacter::ASideScrollingCharacter()
+{
+	PrimaryActorTick.bCanEverTick = true;
+
+	// create the camera component
+	Camera = CreateDefaultSubobject<UCameraComponent>(TEXT("Camera"));
+	Camera->SetupAttachment(RootComponent);
+
+	Camera->SetRelativeLocationAndRotation(FVector(0.0f, 300.0f, 0.0f), FRotator(0.0f, -90.0f, 0.0f));
+
+	// configure the collision capsule
+	GetCapsuleComponent()->SetCapsuleSize(35.0f, 90.0f);
+
+	// configure the Pawn properties
+	bUseControllerRotationYaw = false;
+
+	// configure the character movement component
+	GetCharacterMovement()->GravityScale = 1.75f;
+	GetCharacterMovement()->MaxAcceleration = 1500.0f;
+	GetCharacterMovement()->BrakingFrictionFactor = 1.0f;
+	GetCharacterMovement()->bUseSeparateBrakingFriction = true;
+	GetCharacterMovement()->Mass = 500.0f;
+
+	GetCharacterMovement()->SetWalkableFloorAngle(75.0f);
+	GetCharacterMovement()->MaxWalkSpeed = 500.0f;
+	GetCharacterMovement()->MinAnalogWalkSpeed = 20.0f;
+	GetCharacterMovement()->BrakingDecelerationWalking = 2000.0f;
+	GetCharacterMovement()->bIgnoreBaseRotation = true;
+
+	GetCharacterMovement()->PerchRadiusThreshold = 15.0f;
+	GetCharacterMovement()->LedgeCheckThreshold = 6.0f;
+
+	GetCharacterMovement()->JumpZVelocity = 750.0f;
+	GetCharacterMovement()->AirControl = 1.0f;
+
+	GetCharacterMovement()->RotationRate = FRotator(0.0f, 750.0f, 0.0f);
+	GetCharacterMovement()->bOrientRotationToMovement = true;
+
+	GetCharacterMovement()->SetPlaneConstraintNormal(FVector(0.0f, 1.0f, 0.0f));
+	GetCharacterMovement()->bConstrainToPlane = true;
+
+	// enable double jump
+	JumpMaxCount = 2;
+}
+
+void ASideScrollingCharacter::EndPlay(EEndPlayReason::Type EndPlayReason)
+{
+	Super::EndPlay(EndPlayReason);
+
+	// clear the wall jump timer
+	GetWorld()->GetTimerManager().ClearTimer(WallJumpTimer);
+}
+
+void ASideScrollingCharacter::SetupPlayerInputComponent(class UInputComponent* PlayerInputComponent)
+{
+	Super::SetupPlayerInputComponent(PlayerInputComponent);
+
+	// Set up action bindings
+	if (UEnhancedInputComponent* EnhancedInputComponent = Cast<UEnhancedInputComponent>(PlayerInputComponent))
+	{
+		// Jumping
+		EnhancedInputComponent->BindAction(JumpAction, ETriggerEvent::Started, this, &ASideScrollingCharacter::DoJumpStart);
+		EnhancedInputComponent->BindAction(JumpAction, ETriggerEvent::Completed, this, &ASideScrollingCharacter::DoJumpEnd);
+
+		// Interacting
+		EnhancedInputComponent->BindAction(InteractAction, ETriggerEvent::Triggered, this, &ASideScrollingCharacter::DoInteract);
+
+		// Moving
+		EnhancedInputComponent->BindAction(MoveAction, ETriggerEvent::Triggered, this, &ASideScrollingCharacter::Move);
+
+		// Dropping from platform
+		EnhancedInputComponent->BindAction(DropAction, ETriggerEvent::Triggered, this, &ASideScrollingCharacter::Drop);
+		EnhancedInputComponent->BindAction(DropAction, ETriggerEvent::Completed, this, &ASideScrollingCharacter::DropReleased);
+
+	}
+}
+
+void ASideScrollingCharacter::NotifyHit(class UPrimitiveComponent* MyComp, AActor* Other, class UPrimitiveComponent* OtherComp, bool bSelfMoved, FVector HitLocation, FVector HitNormal, FVector NormalImpulse, const FHitResult& Hit)
+{
+	Super::NotifyHit(MyComp, Other, OtherComp, bSelfMoved, HitLocation, HitNormal, NormalImpulse, Hit);
+
+	// only apply push impulse if we're falling
+	if (!GetCharacterMovement()->IsFalling())
+	{
+		return;
+	}
+
+	// ensure the colliding component is valid
+	if (OtherComp)
+	{
+		// ensure the component is movable and simulating physics
+		if (OtherComp->Mobility == EComponentMobility::Movable && OtherComp->IsSimulatingPhysics())
+		{
+			const FVector PushDir = FVector(ActionValueY > 0.0f ? 1.0f : -1.0f, 0.0f, 0.0f);
+
+			// push the component away
+			OtherComp->AddImpulse(PushDir * JumpPushImpulse, NAME_None, true);
+		}
+	}
+}
+
+void ASideScrollingCharacter::Landed(const FHitResult& Hit)
+{
+	// reset the double jump
+	bHasDoubleJumped = false;
+}
+
+void ASideScrollingCharacter::Move(const FInputActionValue& Value)
+{
+	FVector2D MoveVector = Value.Get<FVector2D>();
+
+	// route the input
+	DoMove(MoveVector.Y);
+}
+
+void ASideScrollingCharacter::Drop(const FInputActionValue& Value)
+{
+	// route the input
+	DoDrop(Value.Get<float>());
+}
+
+void ASideScrollingCharacter::DropReleased(const FInputActionValue& Value)
+{
+	// reset the input
+	DoDrop(0.0f);
+}
+
+void ASideScrollingCharacter::DoMove(float Forward)
+{
+	// is movement temporarily disabled after wall jumping?
+	if (!bHasWallJumped)
+	{
+		// save the movement values
+		ActionValueY = Forward;
+
+		// figure out the movement direction
+		const FVector MoveDir = FVector(1.0f, Forward > 0.0f ? 0.1f : 0.1f, 0.0f);
+
+		// apply the movement input
+		AddMovementInput(MoveDir, Forward);
+	}
+}
+
+void ASideScrollingCharacter::DoDrop(float Value)
+{
+	// save the movement value
+	DropValue = Value;
+}
+
+void ASideScrollingCharacter::DoJumpStart()
+{
+	// handle advanced jump behaviors
+	MultiJump();
+}
+
+void ASideScrollingCharacter::DoJumpEnd()
+{
+	StopJumping();
+}
+
+void ASideScrollingCharacter::DoInteract()
+{
+	// do a sphere trace to look for interactive objects
+	FHitResult OutHit;
+
+	const FVector Start = GetActorLocation();
+	const FVector End = Start + FVector(100.0f, 0.0f, 0.0f);
+
+	FCollisionShape ColSphere;
+	ColSphere.SetSphere(InteractionRadius);
+
+	FCollisionObjectQueryParams ObjectParams;
+	ObjectParams.AddObjectTypesToQuery(ECC_Pawn);
+	ObjectParams.AddObjectTypesToQuery(ECC_WorldDynamic);
+
+	FCollisionQueryParams QueryParams;
+	QueryParams.AddIgnoredActor(this);
+
+	if (GetWorld()->SweepSingleByObjectType(OutHit, Start, End, FQuat::Identity, ObjectParams, ColSphere, QueryParams))
+	{
+		// have we hit an interactable?
+		if (ISideScrollingInteractable* Interactable = Cast<ISideScrollingInteractable>(OutHit.GetActor()))
+		{
+			// interact
+			Interactable->Interaction(this);
+		}
+	}
+}
+
+void ASideScrollingCharacter::MultiJump()
+{
+	// does the user want to drop to a lower platform?
+	if (DropValue > 0.0f)
+	{
+		CheckForSoftCollision();
+		return;
+	}
+
+	// reset the drop value
+	DropValue = 0.0f;
+
+	// if we're grounded, disregard advanced jump logic
+	if (!GetCharacterMovement()->IsFalling())
+	{
+		Jump();
+		return;
+	}
+
+	// if we have a horizontal input, try for wall jump first
+	if (!bHasWallJumped && !FMath::IsNearlyZero(ActionValueY))
+	{
+		// trace ahead of the character for walls
+		FHitResult OutHit;
+
+		const FVector Start = GetActorLocation();
+		const FVector End = Start + (FVector(ActionValueY > 0.0f ? 1.0f : -1.0f, 0.0f, 0.0f) * WallJumpTraceDistance);
+
+		FCollisionQueryParams QueryParams;
+		QueryParams.AddIgnoredActor(this);
+
+		GetWorld()->LineTraceSingleByChannel(OutHit, Start, End, ECC_Visibility, QueryParams);
+
+		if (OutHit.bBlockingHit)
+		{
+			// rotate to the bounce direction
+			const FRotator BounceRot = UKismetMathLibrary::MakeRotFromX(OutHit.ImpactNormal);
+			SetActorRotation(FRotator(0.0f, BounceRot.Yaw, 0.0f));
+
+			// calculate the impulse vector
+			FVector WallJumpImpulse = OutHit.ImpactNormal * WallJumpHorizontalImpulse;
+			WallJumpImpulse.Z = GetCharacterMovement()->JumpZVelocity * WallJumpVerticalMultiplier;
+
+			// launch the character away from the wall
+			LaunchCharacter(WallJumpImpulse, true, true);
+
+			// enable wall jump lockout for a bit
+			bHasWallJumped = true;
+
+			// schedule wall jump lockout reset
+			GetWorld()->GetTimerManager().SetTimer(WallJumpTimer, this, &ASideScrollingCharacter::ResetWallJump, DelayBetweenWallJumps, false);
+
+			return;
+		}
+	}
+
+	// test for double jump only if we haven't already tested for wall jump
+	if (!bHasWallJumped)
+	{
+		// The movement component handles double jump but we still need to manage the flag for animation
+		if (!bHasDoubleJumped)
+		{
+			// raise the double jump flag
+			bHasDoubleJumped = true;
+
+			// let the CMC handle jump
+			Jump();
+		}
+	}
+}
+
+void ASideScrollingCharacter::CheckForSoftCollision()
+{
+	// reset the drop value
+	DropValue = 0.0f;
+
+	// trace down 
+	FHitResult OutHit;
+
+	const FVector Start = GetActorLocation();
+	const FVector End = Start + (FVector::DownVector * SoftCollisionTraceDistance);
+
+	FCollisionObjectQueryParams ObjectParams;
+	ObjectParams.AddObjectTypesToQuery(SoftCollisionObjectType);
+
+	FCollisionQueryParams QueryParams;
+	QueryParams.AddIgnoredActor(this);
+
+	GetWorld()->LineTraceSingleByObjectType(OutHit, Start, End, ObjectParams, QueryParams);
+
+	// did we hit a soft floor?
+	if (OutHit.GetActor())
+	{
+		// drop through the floor
+		SetSoftCollision(true);
+	}
+}
+
+void ASideScrollingCharacter::ResetWallJump()
+{
+	// reset the wall jump flag
+	bHasWallJumped = false;
+}
+
+void ASideScrollingCharacter::SetSoftCollision(bool bEnabled)
+{
+	// enable or disable collision response to the soft collision channel
+	GetCapsuleComponent()->SetCollisionResponseToChannel(SoftCollisionObjectType, bEnabled ? ECR_Ignore : ECR_Block);
+}
+
+bool ASideScrollingCharacter::HasDoubleJumped() const
+{
+	return bHasDoubleJumped;
+}
+
+bool ASideScrollingCharacter::HasWallJumped() const
+{
+	return bHasWallJumped;
+}
diff --git a/Templates/TP_ThirdPerson/Source/TP_ThirdPerson/Variant_SideScrolling/SideScrollingCharacter.h b/Templates/TP_ThirdPerson/Source/TP_ThirdPerson/Variant_SideScrolling/SideScrollingCharacter.h
new file mode 100644
--- /dev/null
+++ b/Templates/TP_ThirdPerson/Source/TP_ThirdPerson/Variant_SideScrolling/SideScrollingCharacter.h
@@ -0,0 +1,170 @@
+// Copyright Epic Games, Inc. All Rights Reserved.
+
+#pragma once
+
+#include "CoreMinimal.h"
+#include "GameFramework/Character.h"
+#include "SideScrollingCharacter.generated.h"
+
+class UCameraComponent;
+class UInputAction;
+struct FInputActionValue;
+
+/**
+ *  A player-controllable character side scrolling game
+ */
+UCLASS(abstract)
+class ASideScrollingCharacter : public ACharacter
+{
+	GENERATED_BODY()
+
+	/** Player camera */
+	UPROPERTY(VisibleAnywhere, BlueprintReadOnly, Category ="Camera", meta = (AllowPrivateAccess = "true"))
+	UCameraComponent* Camera;
+
+protected:
+
+	/** Move Input Action */
+	UPROPERTY(EditAnywhere, BlueprintReadOnly, Category="Input")
+	UInputAction* MoveAction;
+
+	/** Jump Input Action */
+	UPROPERTY(EditAnywhere, BlueprintReadOnly, Category="Input")
+	UInputAction* JumpAction;
+
+	/** Drop from Platform Action */
+	UPROPERTY(EditAnywhere, BlueprintReadOnly, Category="Input")
+	UInputAction* DropAction;
+
+	/** Interact Input Action */
+	UPROPERTY(EditAnywhere, BlueprintReadOnly, Category="Input")
+	UInputAction* InteractAction;
+
+	/** Max distance that interactive objects can be triggered */
+	UPROPERTY(EditAnywhere, Category="Side Scrolling")
+	float InteractionRadius = 200.0f;
+
+	/** Time to disable input after a wall jump to preserve momentum */
+	UPROPERTY(EditAnywhere, Category="Side Scrolling")
+	float DelayBetweenWallJumps = 0.3f;
+
+	/** Impulse to manually push physics objects while we're in midair */
+	UPROPERTY(EditAnywhere, Category="Side Scrolling")
+	float JumpPushImpulse = 600.0f;
+
+	/** Distance to trace ahead of the character for wall jumps */
+	UPROPERTY(EditAnywhere, Category="Side Scrolling")
+	float WallJumpTraceDistance = 50.0f;
+
+	/** Horizontal impulse to apply to the character during wall jumps */
+	UPROPERTY(EditAnywhere, Category="Side Scrolling")
+	float WallJumpHorizontalImpulse = 500.0f;
+
+	/** Multiplies the jump Z velocity for wall jumps. */
+	UPROPERTY(EditAnywhere, Category="Side Scrolling")
+	float WallJumpVerticalMultiplier = 1.4f;
+
+	/** Collision object type to use for soft collision traces (dropping down floors) */
+	UPROPERTY(EditAnywhere, Category="Side Scrolling")
+	TEnumAsByte<ECollisionChannel> SoftCollisionObjectType;
+
+	/** Distance to trace down during soft collision checks */
+	UPROPERTY(EditAnywhere, Category="Side Scrolling")
+	float SoftCollisionTraceDistance = 1000.0f;
+
+	/** Wall jump lockout timer */
+	FTimerHandle WallJumpTimer;
+
+	/** Last captured horizontal movement input value */
+	float ActionValueY = 0.0f;
+
+	/** Last captured platform drop axis value */
+	float DropValue = 0.0f;
+
+	/** If true, this character has already wall jumped */
+	bool bHasWallJumped = false;
+
+	/** If true, this character has already double jumped */
+	bool bHasDoubleJumped = false;
+
+	/** If true, this character is moving along the side scrolling axis */
+	bool bMovingHorizontally = false;
+
+public:
+	
+	/** Constructor */
+	ASideScrollingCharacter();
+
+protected:
+
+	/** Gameplay cleanup */
+	virtual void EndPlay(EEndPlayReason::Type EndPlayReason) override;
+
+	/** Initialize input action bindings */
+	virtual void SetupPlayerInputComponent(class UInputComponent* PlayerInputComponent) override;
+
+	/** Collision handling */
+	virtual void NotifyHit(class UPrimitiveComponent* MyComp, AActor* Other, class UPrimitiveComponent* OtherComp, bool bSelfMoved, FVector HitLocation, FVector HitNormal, FVector NormalImpulse, const FHitResult& Hit) override;
+
+	/** Landing handling */
+	virtual void Landed(const FHitResult& Hit) override;
+
+protected:
+
+	/** Called for movement input */
+	void Move(const FInputActionValue& Value);
+
+	/** Called for drop from platform input */
+	void Drop(const FInputActionValue& Value);
+
+	/** Called for drop from platform input release */
+	void DropReleased(const FInputActionValue& Value);
+
+public:
+
+	/** Handles move inputs from either controls or UI interfaces */
+	UFUNCTION(BlueprintCallable, Category="Input")
+	virtual void DoMove(float Forward);
+
+	/** Handles drop inputs from either controls or UI interfaces */
+	UFUNCTION(BlueprintCallable, Category="Input")
+	virtual void DoDrop(float Value);
+
+	/** Handles jump pressed inputs from either controls or UI interfaces */
+	UFUNCTION(BlueprintCallable, Category="Input")
+	virtual void DoJumpStart();
+
+	/** Handles jump pressed inputs from either controls or UI interfaces */
+	UFUNCTION(BlueprintCallable, Category="Input")
+	virtual void DoJumpEnd();
+
+	/** Handles interact inputs from either controls or UI interfaces */
+	UFUNCTION(BlueprintCallable, Category="Input")
+	virtual void DoInteract();
+
+protected:
+
+	/** Handles advanced jump logic */
+	void MultiJump();
+
+	/** Checks for soft collision with platforms */
+	void CheckForSoftCollision();
+
+	/** Resets wall jump lockout. Called from timer after a wall jump */
+	void ResetWallJump();
+
+public:
+
+	/** Sets the soft collision response. True passes, False blocks */
+	void SetSoftCollision(bool bEnabled);
+
+public:
+
+	/** Returns true if the character has just double jumped */
+	UFUNCTION(BlueprintPure, Category="Side Scrolling")
+	bool HasDoubleJumped() const;
+
+	/** Returns true if the character has just wall jumped */
+	UFUNCTION(BlueprintPure, Category="Side Scrolling")
+	bool HasWallJumped() const;
+};
diff --git a/Templates/TP_ThirdPerson/Source/TP_ThirdPerson/Variant_SideScrolling/SideScrollingGameMode.cpp b/Templates/TP_ThirdPerson/Source/TP_ThirdPerson/Variant_SideScrolling/SideScrollingGameMode.cpp
new file mode 100644
--- /dev/null
+++ b/Templates/TP_ThirdPerson/Source/TP_ThirdPerson/Variant_SideScrolling/SideScrollingGameMode.cpp
@@ -0,0 +1,35 @@
+// Copyright Epic Games, Inc. All Rights Reserved.
+
+
+#include "SideScrollingGameMode.h"
+#include "Kismet/GameplayStatics.h"
+#include "Blueprint/UserWidget.h"
+#include "SideScrollingUI.h"
+#include "SideScrollingPickup.h"
+
+void ASideScrollingGameMode::BeginPlay()
+{
+	Super::BeginPlay();
+
+	// create the game UI
+	APlayerController* OwningPlayer = UGameplayStatics::GetPlayerController(GetWorld(), 0);
+	
+	UserInterface = CreateWidget<USideScrollingUI>(OwningPlayer, UserInterfaceClass);
+
+	check(UserInterface);
+}
+
+void ASideScrollingGameMode::ProcessPickup()
+{
+	// increment the pickups counter
+	++PickupsCollected;
+
+	// if this is the first pickup we collect, show the UI
+	if (PickupsCollected == 1)
+	{
+		UserInterface->AddToViewport(0);
+	}
+
+	// update the pickups counter on the UI
+	UserInterface->UpdatePickups(PickupsCollected);
+}
\ No newline at end of file
diff --git a/Templates/TP_ThirdPerson/Source/TP_ThirdPerson/Variant_SideScrolling/SideScrollingGameMode.h b/Templates/TP_ThirdPerson/Source/TP_ThirdPerson/Variant_SideScrolling/SideScrollingGameMode.h
new file mode 100644
--- /dev/null
+++ b/Templates/TP_ThirdPerson/Source/TP_ThirdPerson/Variant_SideScrolling/SideScrollingGameMode.h
@@ -0,0 +1,44 @@
+// Copyright Epic Games, Inc. All Rights Reserved.
+
+#pragma once
+
+#include "CoreMinimal.h"
+#include "GameFramework/GameModeBase.h"
+#include "SideScrollingGameMode.generated.h"
+
+class USideScrollingUI;
+
+/**
+ *  Simple Side Scrolling Game Mode
+ *  Spawns and manages the game UI
+ *  Counts pickups collected by the player
+ */
+UCLASS(abstract)
+class ASideScrollingGameMode : public AGameModeBase
+{
+	GENERATED_BODY()
+	
+protected:
+
+	/** Class of UI widget to spawn when the game starts */
+	UPROPERTY(EditAnywhere, Category="UI")
+	TSubclassOf<USideScrollingUI> UserInterfaceClass;
+
+	/** User interface widget for the game */
+	UPROPERTY(BlueprintReadOnly, Category="UI")
+	TObjectPtr<USideScrollingUI> UserInterface;
+
+	/** Number of pickups collected by the player */
+	UPROPERTY(BlueprintReadOnly, Category="Picups")
+	int32 PickupsCollected = 0;
+
+protected:
+
+	/** Initialization */
+	virtual void BeginPlay() override;
+
+public:
+
+	/** Receives an interaction event from another actor */
+	virtual void ProcessPickup();
+};
diff --git a/Templates/TP_ThirdPerson/Source/TP_ThirdPerson/Variant_SideScrolling/SideScrollingInteractable.cpp b/Templates/TP_ThirdPerson/Source/TP_ThirdPerson/Variant_SideScrolling/SideScrollingInteractable.cpp
new file mode 100644
--- /dev/null
+++ b/Templates/TP_ThirdPerson/Source/TP_ThirdPerson/Variant_SideScrolling/SideScrollingInteractable.cpp
@@ -0,0 +1,6 @@
+// Copyright Epic Games, Inc. All Rights Reserved.
+
+
+#include "SideScrollingInteractable.h"
+
+// Add default functionality here for any IInteractable functions that are not pure virtual.
diff --git a/Templates/TP_ThirdPerson/Source/TP_ThirdPerson/Variant_SideScrolling/SideScrollingInteractable.h b/Templates/TP_ThirdPerson/Source/TP_ThirdPerson/Variant_SideScrolling/SideScrollingInteractable.h
new file mode 100644
--- /dev/null
+++ b/Templates/TP_ThirdPerson/Source/TP_ThirdPerson/Variant_SideScrolling/SideScrollingInteractable.h
@@ -0,0 +1,31 @@
+// Copyright Epic Games, Inc. All Rights Reserved.
+
+#pragma once
+
+#include "CoreMinimal.h"
+#include "UObject/Interface.h"
+#include "SideScrollingInteractable.generated.h"
+
+/**
+ *  
+ */
+UINTERFACE(MinimalAPI, NotBlueprintable)
+class USideScrollingInteractable : public UInterface
+{
+	GENERATED_BODY()
+};
+
+/**
+ *  Simple interface to allow Actors to interact without having knowledge of their internal implementation.
+ */
+class ISideScrollingInteractable
+{
+	GENERATED_BODY()
+
+public:
+
+	/** Triggers an interaction by the provided Actor */
+	UFUNCTION(BlueprintCallable, Category="Interactable")
+	virtual void Interaction(AActor* Interactor) = 0;
+
+};
diff --git a/Templates/TP_ThirdPerson/Source/TP_ThirdPerson/Variant_SideScrolling/SideScrollingPlayerController.cpp b/Templates/TP_ThirdPerson/Source/TP_ThirdPerson/Variant_SideScrolling/SideScrollingPlayerController.cpp
new file mode 100644
--- /dev/null
+++ b/Templates/TP_ThirdPerson/Source/TP_ThirdPerson/Variant_SideScrolling/SideScrollingPlayerController.cpp
@@ -0,0 +1,50 @@
+// Copyright Epic Games, Inc. All Rights Reserved.
+
+
+#include "SideScrollingPlayerController.h"
+#include "EnhancedInputSubsystems.h"
+#include "InputMappingContext.h"
+#include "Kismet/GameplayStatics.h"
+#include "GameFramework/PlayerStart.h"
+#include "SideScrollingCharacter.h"
+#include "Engine/LocalPlayer.h"
+#include "Engine/World.h"
+
+void ASideScrollingPlayerController::SetupInputComponent()
+{
+	// add the input mapping context
+	if (UEnhancedInputLocalPlayerSubsystem* Subsystem = ULocalPlayer::GetSubsystem<UEnhancedInputLocalPlayerSubsystem>(GetLocalPlayer()))
+	{
+		for (UInputMappingContext* CurrentContext : DefaultMappingContexts)
+		{
+			Subsystem->AddMappingContext(CurrentContext, 0);
+		}
+	}
+}
+
+void ASideScrollingPlayerController::OnPossess(APawn* InPawn)
+{
+	Super::OnPossess(InPawn);
+
+	// subscribe to the pawn's OnDestroyed delegate
+	InPawn->OnDestroyed.AddDynamic(this, &ASideScrollingPlayerController::OnPawnDestroyed);
+}
+
+void ASideScrollingPlayerController::OnPawnDestroyed(AActor* DestroyedActor)
+{
+	// find the player start
+	TArray<AActor*> ActorList;
+	UGameplayStatics::GetAllActorsOfClass(GetWorld(), APlayerStart::StaticClass(), ActorList);
+
+	if (ActorList.Num() > 0)
+	{
+		// spawn a character at the player start
+		const FTransform SpawnTransform = ActorList[0]->GetActorTransform();
+
+		if (ASideScrollingCharacter* RespawnedCharacter = GetWorld()->SpawnActor<ASideScrollingCharacter>(CharacterClass, SpawnTransform))
+		{
+			// possess the character
+			Possess(RespawnedCharacter);
+		}
+	}
+}
diff --git a/Templates/TP_ThirdPerson/Source/TP_ThirdPerson/Variant_SideScrolling/SideScrollingPlayerController.h b/Templates/TP_ThirdPerson/Source/TP_ThirdPerson/Variant_SideScrolling/SideScrollingPlayerController.h
new file mode 100644
--- /dev/null
+++ b/Templates/TP_ThirdPerson/Source/TP_ThirdPerson/Variant_SideScrolling/SideScrollingPlayerController.h
@@ -0,0 +1,45 @@
+// Copyright Epic Games, Inc. All Rights Reserved.
+
+#pragma once
+
+#include "CoreMinimal.h"
+#include "GameFramework/PlayerController.h"
+#include "EnhancedInput/Public/InputAction.h"
+#include "SideScrollingPlayerController.generated.h"
+
+class ASideScrollingCharacter;
+class UInputMappingContext;
+
+/**
+ *  A simple Side Scrolling Player Controller
+ *  Manages input mappings
+ *  Respawns the player pawn at the player start if it is destroyed
+ */
+UCLASS(abstract)
+class ASideScrollingPlayerController : public APlayerController
+{
+	GENERATED_BODY()
+	
+protected:
+
+	/** Input mapping context for this player */
+	UPROPERTY(EditAnywhere, Category="Input")
+	TArray<UInputMappingContext*> DefaultMappingContexts;
+
+	/** Character class to respawn when the possessed pawn is destroyed */
+	UPROPERTY(EditAnywhere, Category="Respawn")
+	TSubclassOf<ASideScrollingCharacter> CharacterClass;
+
+protected:
+
+	/** Initialize input bindings */
+	virtual void SetupInputComponent() override;
+
+	/** Pawn initialization */
+	virtual void OnPossess(APawn* InPawn) override;
+
+	/** Called if the possessed pawn is destroyed */
+	UFUNCTION()
+	void OnPawnDestroyed(AActor* DestroyedActor);
+
+};
diff --git a/Templates/TP_ThirdPerson/Source/TP_ThirdPerson/Variant_SideScrolling/SideScrollingUI.cpp b/Templates/TP_ThirdPerson/Source/TP_ThirdPerson/Variant_SideScrolling/SideScrollingUI.cpp
new file mode 100644
--- /dev/null
+++ b/Templates/TP_ThirdPerson/Source/TP_ThirdPerson/Variant_SideScrolling/SideScrollingUI.cpp
@@ -0,0 +1,5 @@
+// Copyright Epic Games, Inc. All Rights Reserved.
+
+
+#include "SideScrollingUI.h"
+
diff --git a/Templates/TP_ThirdPerson/Source/TP_ThirdPerson/Variant_SideScrolling/SideScrollingUI.h b/Templates/TP_ThirdPerson/Source/TP_ThirdPerson/Variant_SideScrolling/SideScrollingUI.h
new file mode 100644
--- /dev/null
+++ b/Templates/TP_ThirdPerson/Source/TP_ThirdPerson/Variant_SideScrolling/SideScrollingUI.h
@@ -0,0 +1,23 @@
+// Copyright Epic Games, Inc. All Rights Reserved.
+
+#pragma once
+
+#include "CoreMinimal.h"
+#include "Blueprint/UserWidget.h"
+#include "SideScrollingUI.generated.h"
+
+/**
+ *  Simple Side Scrolling game UI
+ *  Displays and manages a pickup counter
+ */
+UCLASS(abstract)
+class USideScrollingUI : public UUserWidget
+{
+	GENERATED_BODY()
+	
+public:
+
+	/** Update the widget's pickup counter */
+	UFUNCTION(BlueprintImplementableEvent, Category="UI")
+	void UpdatePickups(int32 Amount);
+};
diff --git a/Templates/TP_ThirdPerson/TP_ThirdPerson.uproject b/Templates/TP_ThirdPerson/TP_ThirdPerson.uproject
--- a/Templates/TP_ThirdPerson/TP_ThirdPerson.uproject
+++ b/Templates/TP_ThirdPerson/TP_ThirdPerson.uproject
@@ -7,7 +7,12 @@
 		{
 			"Name": "TP_ThirdPerson",
 			"Type": "Runtime",
-			"LoadingPhase": "Default"
+			"LoadingPhase": "Default",
+			"AdditionalDependencies": [
+				"Engine",
+				"AIModule",
+				"UMG"
+			]
 		}
 	],
 	"Plugins": [
@@ -17,6 +22,14 @@
 			"TargetAllowList": [
 				"Editor"
 			]
+		},
+		{
+			"Name": "StateTree",
+			"Enabled": true
+		},
+		{
+			"Name": "GameplayStateTree",
+			"Enabled": true
 		}
 	]
 }
\ No newline at end of file
