Templates
- Fixed config discrepancies between BP and C++ template versions.
- Updated default Mac RHI to SM6
#rnx

diff --git a/Templates/TP_Blank/Config/DefaultEngine.ini b/Templates/TP_Blank/Config/DefaultEngine.ini
--- a/Templates/TP_Blank/Config/DefaultEngine.ini
+++ b/Templates/TP_Blank/Config/DefaultEngine.ini
@@ -37,4 +37,8 @@ SoundCueCookQualityIndex=-1
 
 [/Script/LinuxTargetPlatform.LinuxTargetSettings]
 -TargetedRHIs=SF_VULKAN_SM5
-+TargetedRHIs=SF_VULKAN_SM6
\ No newline at end of file
++TargetedRHIs=SF_VULKAN_SM6
+
+[/Script/MacTargetPlatform.MacTargetSettings]
+-TargetedRHIs=SF_METAL_SM5
++TargetedRHIs=SF_METAL_SM6
\ No newline at end of file
diff --git a/Templates/TP_BlankBP/Config/DefaultEngine.ini b/Templates/TP_BlankBP/Config/DefaultEngine.ini
--- a/Templates/TP_BlankBP/Config/DefaultEngine.ini
+++ b/Templates/TP_BlankBP/Config/DefaultEngine.ini
@@ -37,4 +37,8 @@ SoundCueCookQualityIndex=-1
 
 [/Script/LinuxTargetPlatform.LinuxTargetSettings]
 -TargetedRHIs=SF_VULKAN_SM5
-+TargetedRHIs=SF_VULKAN_SM6
\ No newline at end of file
++TargetedRHIs=SF_VULKAN_SM6
+
+[/Script/MacTargetPlatform.MacTargetSettings]
+-TargetedRHIs=SF_METAL_SM5
++TargetedRHIs=SF_METAL_SM6
\ No newline at end of file
diff --git a/Templates/TP_FirstPerson/Config/DefaultEngine.ini b/Templates/TP_FirstPerson/Config/DefaultEngine.ini
--- a/Templates/TP_FirstPerson/Config/DefaultEngine.ini
+++ b/Templates/TP_FirstPerson/Config/DefaultEngine.ini
@@ -57,6 +57,10 @@ SoundCueCookQualityIndex=-1
 -TargetedRHIs=SF_VULKAN_SM5
 +TargetedRHIs=SF_VULKAN_SM6
 
+[/Script/MacTargetPlatform.MacTargetSettings]
+-TargetedRHIs=SF_METAL_SM5
++TargetedRHIs=SF_METAL_SM6
+
 [/Script/AIModule.AISystem]
 bForgetStaleActors=True
 
diff --git a/Templates/TP_FirstPersonBP/Config/DefaultEngine.ini b/Templates/TP_FirstPersonBP/Config/DefaultEngine.ini
--- a/Templates/TP_FirstPersonBP/Config/DefaultEngine.ini
+++ b/Templates/TP_FirstPersonBP/Config/DefaultEngine.ini
@@ -20,25 +20,9 @@ GlobalDefaultServerGameMode=None
 r.ReflectionMethod=1
 r.GenerateMeshDistanceFields=True
 r.DynamicGlobalIlluminationMethod=1
-r.Mobile.EnableNoPrecomputedLightingCSMShader=1
-r.DefaultFeature.AutoExposure.ExtendDefaultLuminanceRange=True
 r.DefaultFeature.AutoExposure.ExtendDefaultLuminanceRange=true
-
-r.Shadow.Virtual.Enable=1
-
-r.DefaultFeature.LocalExposure.HighlightContrastScale=0.8
-
-r.DefaultFeature.LocalExposure.ShadowContrastScale=0.8
-r.DefaultFeature.LightUnits=2
 r.AllowStaticLighting=False
-RedChromaticityCoordinate=(X=0.640000,Y=0.330000)
-GreenChromaticityCoordinate=(X=0.300000,Y=0.600000)
-BlueChromaticityCoordinate=(X=0.150000,Y=0.060000)
-WhiteChromaticityCoordinate=(X=0.312700,Y=0.329000)
-WorkingColorSpaceChoice=sRGB
-r.LegacyLuminanceFactors=False
-r.Lumen.ScreenTracingSource=0
-r.FirstPerson.SelfShadow=1
+
 
 [/Script/WindowsTargetPlatform.WindowsTargetSettings]
 DefaultGraphicsRHI=DefaultGraphicsRHI_DX12
@@ -46,18 +30,40 @@ DefaultGraphicsRHI=DefaultGraphicsRHI_DX12
 +D3D12TargetedShaderFormats=PCD3D_SM6
 -D3D11TargetedShaderFormats=PCD3D_SM5
 +D3D11TargetedShaderFormats=PCD3D_SM5
+Compiler=Default
+AudioSampleRate=48000
+AudioCallbackBufferFrameSize=1024
+AudioNumBuffersToEnqueue=1
+AudioMaxChannels=0
+AudioNumSourceWorkers=4
+SpatializationPlugin=
+SourceDataOverridePlugin=
+ReverbPlugin=
+OcclusionPlugin=
+CompressionOverrides=(bOverrideCompressionTimes=False,DurationThreshold=5.000000,MaxNumRandomBranches=0,SoundCueQualityIndex=0)
+CacheSizeKB=65536
+MaxChunkSizeOverrideKB=0
+bResampleForDevice=False
+MaxSampleRate=48000.000000
+HighSampleRate=32000.000000
+MedSampleRate=24000.000000
+LowSampleRate=12000.000000
+MinSampleRate=8000.000000
+CompressionQualityModifier=1.000000
+AutoStreamingThreshold=0.000000
+SoundCueCookQualityIndex=-1
 
-[/Script/HardwareTargeting.HardwareTargetingSettings]
-TargetedHardwareClass=Desktop
-AppliedTargetedHardwareClass=Desktop
-DefaultGraphicsPerformance=Maximum
-AppliedDefaultGraphicsPerformance=Maximum
+[/Script/LinuxTargetPlatform.LinuxTargetSettings]
+-TargetedRHIs=SF_VULKAN_SM5
++TargetedRHIs=SF_VULKAN_SM6
 
-[/Script/Engine.Engine]
-NearClipPlane=5.000000
+[/Script/MacTargetPlatform.MacTargetSettings]
+-TargetedRHIs=SF_METAL_SM5
++TargetedRHIs=SF_METAL_SM6
 
 [/Script/AIModule.AISystem]
 bForgetStaleActors=True
 
-[/Script/Engine.UserInterfaceSettings]
-UIScaleRule=ShortestSide
\ No newline at end of file
+[/Script/Engine.Engine]
+NearClipPlane=5.000000
+
diff --git a/Templates/TP_ThirdPerson/Config/DefaultEngine.ini b/Templates/TP_ThirdPerson/Config/DefaultEngine.ini
--- a/Templates/TP_ThirdPerson/Config/DefaultEngine.ini
+++ b/Templates/TP_ThirdPerson/Config/DefaultEngine.ini
@@ -45,59 +45,6 @@ SoundCueCookQualityIndex=-1
 -TargetedRHIs=SF_VULKAN_SM5
 +TargetedRHIs=SF_VULKAN_SM6
 
-[/Script/Engine.CollisionProfile]
--Profiles=(Name="NoCollision",CollisionEnabled=NoCollision,ObjectTypeName="WorldStatic",CustomResponses=((Channel="Visibility",Response=ECR_Ignore),(Channel="Camera",Response=ECR_Ignore)),HelpMessage="No collision",bCanModify=False)
--Profiles=(Name="BlockAll",CollisionEnabled=QueryAndPhysics,ObjectTypeName="WorldStatic",CustomResponses=,HelpMessage="WorldStatic object that blocks all actors by default. All new custom channels will use its own default response. ",bCanModify=False)
--Profiles=(Name="OverlapAll",CollisionEnabled=QueryOnly,ObjectTypeName="WorldStatic",CustomResponses=((Channel="WorldStatic",Response=ECR_Overlap),(Channel="Pawn",Response=ECR_Overlap),(Channel="Visibility",Response=ECR_Overlap),(Channel="WorldDynamic",Response=ECR_Overlap),(Channel="Camera",Response=ECR_Overlap),(Channel="PhysicsBody",Response=ECR_Overlap),(Channel="Vehicle",Response=ECR_Overlap),(Channel="Destructible",Response=ECR_Overlap)),HelpMessage="WorldStatic object that overlaps all actors by default. All new custom channels will use its own default response. ",bCanModify=False)
--Profiles=(Name="BlockAllDynamic",CollisionEnabled=QueryAndPhysics,ObjectTypeName="WorldDynamic",CustomResponses=,HelpMessage="WorldDynamic object that blocks all actors by default. All new custom channels will use its own default response. ",bCanModify=False)
--Profiles=(Name="OverlapAllDynamic",CollisionEnabled=QueryOnly,ObjectTypeName="WorldDynamic",CustomResponses=((Channel="WorldStatic",Response=ECR_Overlap),(Channel="Pawn",Response=ECR_Overlap),(Channel="Visibility",Response=ECR_Overlap),(Channel="WorldDynamic",Response=ECR_Overlap),(Channel="Camera",Response=ECR_Overlap),(Channel="PhysicsBody",Response=ECR_Overlap),(Channel="Vehicle",Response=ECR_Overlap),(Channel="Destructible",Response=ECR_Overlap)),HelpMessage="WorldDynamic object that overlaps all actors by default. All new custom channels will use its own default response. ",bCanModify=False)
--Profiles=(Name="IgnoreOnlyPawn",CollisionEnabled=QueryOnly,ObjectTypeName="WorldDynamic",CustomResponses=((Channel="Pawn",Response=ECR_Ignore),(Channel="Vehicle",Response=ECR_Ignore)),HelpMessage="WorldDynamic object that ignores Pawn and Vehicle. All other channels will be set to default.",bCanModify=False)
--Profiles=(Name="OverlapOnlyPawn",CollisionEnabled=QueryOnly,ObjectTypeName="WorldDynamic",CustomResponses=((Channel="Pawn",Response=ECR_Overlap),(Channel="Vehicle",Response=ECR_Overlap),(Channel="Camera",Response=ECR_Ignore)),HelpMessage="WorldDynamic object that overlaps Pawn, Camera, and Vehicle. All other channels will be set to default. ",bCanModify=False)
--Profiles=(Name="Pawn",CollisionEnabled=QueryAndPhysics,ObjectTypeName="Pawn",CustomResponses=((Channel="Visibility",Response=ECR_Ignore)),HelpMessage="Pawn object. Can be used for capsule of any playerable character or AI. ",bCanModify=False)
--Profiles=(Name="Spectator",CollisionEnabled=QueryOnly,ObjectTypeName="Pawn",CustomResponses=((Channel="WorldStatic",Response=ECR_Block),(Channel="Pawn",Response=ECR_Ignore),(Channel="Visibility",Response=ECR_Ignore),(Channel="WorldDynamic",Response=ECR_Ignore),(Channel="Camera",Response=ECR_Ignore),(Channel="PhysicsBody",Response=ECR_Ignore),(Channel="Vehicle",Response=ECR_Ignore),(Channel="Destructible",Response=ECR_Ignore)),HelpMessage="Pawn object that ignores all other actors except WorldStatic.",bCanModify=False)
--Profiles=(Name="CharacterMesh",CollisionEnabled=QueryOnly,ObjectTypeName="Pawn",CustomResponses=((Channel="Pawn",Response=ECR_Ignore),(Channel="Vehicle",Response=ECR_Ignore),(Channel="Visibility",Response=ECR_Ignore)),HelpMessage="Pawn object that is used for Character Mesh. All other channels will be set to default.",bCanModify=False)
--Profiles=(Name="PhysicsActor",CollisionEnabled=QueryAndPhysics,ObjectTypeName="PhysicsBody",CustomResponses=,HelpMessage="Simulating actors",bCanModify=False)
--Profiles=(Name="Destructible",CollisionEnabled=QueryAndPhysics,ObjectTypeName="Destructible",CustomResponses=,HelpMessage="Destructible actors",bCanModify=False)
--Profiles=(Name="InvisibleWall",CollisionEnabled=QueryAndPhysics,ObjectTypeName="WorldStatic",CustomResponses=((Channel="Visibility",Response=ECR_Ignore)),HelpMessage="WorldStatic object that is invisible.",bCanModify=False)
--Profiles=(Name="InvisibleWallDynamic",CollisionEnabled=QueryAndPhysics,ObjectTypeName="WorldDynamic",CustomResponses=((Channel="Visibility",Response=ECR_Ignore)),HelpMessage="WorldDynamic object that is invisible.",bCanModify=False)
--Profiles=(Name="Trigger",CollisionEnabled=QueryOnly,ObjectTypeName="WorldDynamic",CustomResponses=((Channel="WorldStatic",Response=ECR_Overlap),(Channel="Pawn",Response=ECR_Overlap),(Channel="Visibility",Response=ECR_Ignore),(Channel="WorldDynamic",Response=ECR_Overlap),(Channel="Camera",Response=ECR_Overlap),(Channel="PhysicsBody",Response=ECR_Overlap),(Channel="Vehicle",Response=ECR_Overlap),(Channel="Destructible",Response=ECR_Overlap)),HelpMessage="WorldDynamic object that is used for trigger. All other channels will be set to default.",bCanModify=False)
--Profiles=(Name="Ragdoll",CollisionEnabled=QueryAndPhysics,ObjectTypeName="PhysicsBody",CustomResponses=((Channel="Pawn",Response=ECR_Ignore),(Channel="Visibility",Response=ECR_Ignore)),HelpMessage="Simulating Skeletal Mesh Component. All other channels will be set to default.",bCanModify=False)
--Profiles=(Name="Vehicle",CollisionEnabled=QueryAndPhysics,ObjectTypeName="Vehicle",CustomResponses=,HelpMessage="Vehicle object that blocks Vehicle, WorldStatic, and WorldDynamic. All other channels will be set to default.",bCanModify=False)
--Profiles=(Name="UI",CollisionEnabled=QueryOnly,ObjectTypeName="WorldDynamic",CustomResponses=((Channel="WorldStatic",Response=ECR_Overlap),(Channel="Pawn",Response=ECR_Overlap),(Channel="Visibility",Response=ECR_Block),(Channel="WorldDynamic",Response=ECR_Overlap),(Channel="Camera",Response=ECR_Overlap),(Channel="PhysicsBody",Response=ECR_Overlap),(Channel="Vehicle",Response=ECR_Overlap),(Channel="Destructible",Response=ECR_Overlap)),HelpMessage="WorldStatic object that overlaps all actors by default. All new custom channels will use its own default response. ",bCanModify=False)
-+Profiles=(Name="NoCollision",CollisionEnabled=NoCollision,bCanModify=False,ObjectTypeName="WorldStatic",CustomResponses=((Channel="Visibility",Response=ECR_Ignore),(Channel="Camera",Response=ECR_Ignore)),HelpMessage="No collision")
-+Profiles=(Name="BlockAll",CollisionEnabled=QueryAndPhysics,bCanModify=False,ObjectTypeName="WorldStatic",CustomResponses=,HelpMessage="WorldStatic object that blocks all actors by default. All new custom channels will use its own default response. ")
-+Profiles=(Name="OverlapAll",CollisionEnabled=QueryOnly,bCanModify=False,ObjectTypeName="WorldStatic",CustomResponses=((Channel="WorldStatic",Response=ECR_Overlap),(Channel="Pawn",Response=ECR_Overlap),(Channel="Visibility",Response=ECR_Overlap),(Channel="WorldDynamic",Response=ECR_Overlap),(Channel="Camera",Response=ECR_Overlap),(Channel="PhysicsBody",Response=ECR_Overlap),(Channel="Vehicle",Response=ECR_Overlap),(Channel="Destructible",Response=ECR_Overlap)),HelpMessage="WorldStatic object that overlaps all actors by default. All new custom channels will use its own default response. ")
-+Profiles=(Name="BlockAllDynamic",CollisionEnabled=QueryAndPhysics,bCanModify=False,ObjectTypeName="WorldDynamic",CustomResponses=,HelpMessage="WorldDynamic object that blocks all actors by default. All new custom channels will use its own default response. ")
-+Profiles=(Name="OverlapAllDynamic",CollisionEnabled=QueryOnly,bCanModify=False,ObjectTypeName="WorldDynamic",CustomResponses=((Channel="WorldStatic",Response=ECR_Overlap),(Channel="Pawn",Response=ECR_Overlap),(Channel="Visibility",Response=ECR_Overlap),(Channel="WorldDynamic",Response=ECR_Overlap),(Channel="Camera",Response=ECR_Overlap),(Channel="PhysicsBody",Response=ECR_Overlap),(Channel="Vehicle",Response=ECR_Overlap),(Channel="Destructible",Response=ECR_Overlap)),HelpMessage="WorldDynamic object that overlaps all actors by default. All new custom channels will use its own default response. ")
-+Profiles=(Name="IgnoreOnlyPawn",CollisionEnabled=QueryOnly,bCanModify=False,ObjectTypeName="WorldDynamic",CustomResponses=((Channel="Pawn",Response=ECR_Ignore),(Channel="Vehicle",Response=ECR_Ignore)),HelpMessage="WorldDynamic object that ignores Pawn and Vehicle. All other channels will be set to default.")
-+Profiles=(Name="OverlapOnlyPawn",CollisionEnabled=QueryOnly,bCanModify=False,ObjectTypeName="WorldDynamic",CustomResponses=((Channel="Pawn",Response=ECR_Overlap),(Channel="Vehicle",Response=ECR_Overlap),(Channel="Camera",Response=ECR_Ignore)),HelpMessage="WorldDynamic object that overlaps Pawn, Camera, and Vehicle. All other channels will be set to default. ")
-+Profiles=(Name="Pawn",CollisionEnabled=QueryAndPhysics,bCanModify=False,ObjectTypeName="Pawn",CustomResponses=((Channel="Visibility",Response=ECR_Ignore)),HelpMessage="Pawn object. Can be used for capsule of any playerable character or AI. ")
-+Profiles=(Name="Spectator",CollisionEnabled=QueryOnly,bCanModify=False,ObjectTypeName="Pawn",CustomResponses=((Channel="WorldStatic"),(Channel="Pawn",Response=ECR_Ignore),(Channel="Visibility",Response=ECR_Ignore),(Channel="WorldDynamic",Response=ECR_Ignore),(Channel="Camera",Response=ECR_Ignore),(Channel="PhysicsBody",Response=ECR_Ignore),(Channel="Vehicle",Response=ECR_Ignore),(Channel="Destructible",Response=ECR_Ignore)),HelpMessage="Pawn object that ignores all other actors except WorldStatic.")
-+Profiles=(Name="CharacterMesh",CollisionEnabled=QueryOnly,bCanModify=False,ObjectTypeName="Pawn",CustomResponses=((Channel="Pawn",Response=ECR_Ignore),(Channel="Vehicle",Response=ECR_Ignore),(Channel="Visibility",Response=ECR_Ignore)),HelpMessage="Pawn object that is used for Character Mesh. All other channels will be set to default.")
-+Profiles=(Name="PhysicsActor",CollisionEnabled=QueryAndPhysics,bCanModify=False,ObjectTypeName="PhysicsBody",CustomResponses=,HelpMessage="Simulating actors")
-+Profiles=(Name="Destructible",CollisionEnabled=QueryAndPhysics,bCanModify=False,ObjectTypeName="Destructible",CustomResponses=,HelpMessage="Destructible actors")
-+Profiles=(Name="InvisibleWall",CollisionEnabled=QueryAndPhysics,bCanModify=False,ObjectTypeName="WorldStatic",CustomResponses=((Channel="Visibility",Response=ECR_Ignore)),HelpMessage="WorldStatic object that is invisible.")
-+Profiles=(Name="InvisibleWallDynamic",CollisionEnabled=QueryAndPhysics,bCanModify=False,ObjectTypeName="WorldDynamic",CustomResponses=((Channel="Visibility",Response=ECR_Ignore)),HelpMessage="WorldDynamic object that is invisible.")
-+Profiles=(Name="Trigger",CollisionEnabled=QueryOnly,bCanModify=False,ObjectTypeName="WorldDynamic",CustomResponses=((Channel="WorldStatic",Response=ECR_Overlap),(Channel="Pawn",Response=ECR_Overlap),(Channel="Visibility",Response=ECR_Ignore),(Channel="WorldDynamic",Response=ECR_Overlap),(Channel="Camera",Response=ECR_Overlap),(Channel="PhysicsBody",Response=ECR_Overlap),(Channel="Vehicle",Response=ECR_Overlap),(Channel="Destructible",Response=ECR_Overlap)),HelpMessage="WorldDynamic object that is used for trigger. All other channels will be set to default.")
-+Profiles=(Name="Ragdoll",CollisionEnabled=QueryAndPhysics,bCanModify=False,ObjectTypeName="PhysicsBody",CustomResponses=((Channel="Pawn",Response=ECR_Ignore),(Channel="Visibility",Response=ECR_Ignore)),HelpMessage="Simulating Skeletal Mesh Component. All other channels will be set to default.")
-+Profiles=(Name="Vehicle",CollisionEnabled=QueryAndPhysics,bCanModify=False,ObjectTypeName="Vehicle",CustomResponses=,HelpMessage="Vehicle object that blocks Vehicle, WorldStatic, and WorldDynamic. All other channels will be set to default.")
-+Profiles=(Name="UI",CollisionEnabled=QueryOnly,bCanModify=False,ObjectTypeName="WorldDynamic",CustomResponses=((Channel="WorldStatic",Response=ECR_Overlap),(Channel="Pawn",Response=ECR_Overlap),(Channel="Visibility"),(Channel="WorldDynamic",Response=ECR_Overlap),(Channel="Camera",Response=ECR_Overlap),(Channel="PhysicsBody",Response=ECR_Overlap),(Channel="Vehicle",Response=ECR_Overlap),(Channel="Destructible",Response=ECR_Overlap)),HelpMessage="WorldStatic object that overlaps all actors by default. All new custom channels will use its own default response. ")
-+DefaultChannelResponses=(Channel=ECC_GameTraceChannel1,DefaultResponse=ECR_Block,bTraceType=False,bStaticObject=False,Name="SoftCollision")
--ProfileRedirects=(OldName="BlockingVolume",NewName="InvisibleWall")
--ProfileRedirects=(OldName="InterpActor",NewName="IgnoreOnlyPawn")
--ProfileRedirects=(OldName="StaticMeshComponent",NewName="BlockAllDynamic")
--ProfileRedirects=(OldName="SkeletalMeshActor",NewName="PhysicsActor")
--ProfileRedirects=(OldName="InvisibleActor",NewName="InvisibleWallDynamic")
-+ProfileRedirects=(OldName="BlockingVolume",NewName="InvisibleWall")
-+ProfileRedirects=(OldName="InterpActor",NewName="IgnoreOnlyPawn")
-+ProfileRedirects=(OldName="StaticMeshComponent",NewName="BlockAllDynamic")
-+ProfileRedirects=(OldName="SkeletalMeshActor",NewName="PhysicsActor")
-+ProfileRedirects=(OldName="InvisibleActor",NewName="InvisibleWallDynamic")
--CollisionChannelRedirects=(OldName="Static",NewName="WorldStatic")
--CollisionChannelRedirects=(OldName="Dynamic",NewName="WorldDynamic")
--CollisionChannelRedirects=(OldName="VehicleMovement",NewName="Vehicle")
--CollisionChannelRedirects=(OldName="PawnMovement",NewName="Pawn")
-+CollisionChannelRedirects=(OldName="Static",NewName="WorldStatic")
-+CollisionChannelRedirects=(OldName="Dynamic",NewName="WorldDynamic")
-+CollisionChannelRedirects=(OldName="VehicleMovement",NewName="Vehicle")
-+CollisionChannelRedirects=(OldName="PawnMovement",NewName="Pawn")
+[/Script/MacTargetPlatform.MacTargetSettings]
+-TargetedRHIs=SF_METAL_SM5
++TargetedRHIs=SF_METAL_SM6
\ No newline at end of file
diff --git a/Templates/TP_ThirdPersonBP/Config/DefaultEngine.ini b/Templates/TP_ThirdPersonBP/Config/DefaultEngine.ini
--- a/Templates/TP_ThirdPersonBP/Config/DefaultEngine.ini
+++ b/Templates/TP_ThirdPersonBP/Config/DefaultEngine.ini
@@ -52,3 +52,7 @@ SoundCueCookQualityIndex=-1
 [/Script/LinuxTargetPlatform.LinuxTargetSettings]
 -TargetedRHIs=SF_VULKAN_SM5
 +TargetedRHIs=SF_VULKAN_SM6
+
+[/Script/MacTargetPlatform.MacTargetSettings]
+-TargetedRHIs=SF_METAL_SM5
++TargetedRHIs=SF_METAL_SM6
\ No newline at end of file
diff --git a/Templates/TP_TopDown/Config/DefaultEngine.ini b/Templates/TP_TopDown/Config/DefaultEngine.ini
--- a/Templates/TP_TopDown/Config/DefaultEngine.ini
+++ b/Templates/TP_TopDown/Config/DefaultEngine.ini
@@ -3,7 +3,50 @@ GameDefaultMap=/Game/TopDown/Lvl_TopDown.Lvl_TopDown
 EditorStartupMap=/Game/TopDown/Lvl_TopDown.Lvl_TopDown
 GlobalDefaultGameMode=/Game/TopDown/Blueprints/BP_TopDownGameMode.BP_TopDownGameMode_C
 
+[/Script/NavigationSystem.RecastNavMesh]
+bDrawPolyEdges=False
+bDistinctlyDrawTilesBeingBuilt=True
+DrawOffset=10.000000
+bFixedTilePoolSize=False
+TilePoolSize=1024
+TileSizeUU=1000.000000
+CellSize=19.000000
+CellHeight=10.000000
+AgentRadius=34.000000
+AgentHeight=144.000000
+AgentMaxHeight=160.000000
+AgentMaxSlope=44.000000
+AgentMaxStepHeight=35.000000
+MinRegionArea=0.000000
+MergeRegionSize=400.000000
+MaxSimplificationError=1.300000
+MaxSimultaneousTileGenerationJobsCount=1024
+TileNumberHardLimit=1048576
+DefaultDrawDistance=5000.000000
+DefaultMaxSearchNodes=2048.000000
+DefaultMaxHierarchicalSearchNodes=2048.000000
+RegionPartitioning=Watershed
+LayerPartitioning=Watershed
+RegionChunkSplits=2
+LayerChunkSplits=2
+bSortNavigationAreasByCost=False
+bPerformVoxelFiltering=True
+bMarkLowHeightAreas=False
+bDoFullyAsyncNavDataGathering=False
+bUseBetterOffsetsFromCorners=True
+bUseVirtualFilters=True
+bUseVoxelCache=False
+TileSetUpdateInterval=1.000000
+HeuristicScale=0.999000
+VerticalDeviationFromGroundCompensation=0.000000
+bForceRebuildOnLoad=True
+RuntimeGeneration=Dynamic
+
 [/Script/Engine.RendererSettings]
+r.ReflectionMethod=1
+r.GenerateMeshDistanceFields=True
+r.DynamicGlobalIlluminationMethod=1
+r.Shadow.Virtual.Enable=1
 r.DefaultFeature.AutoExposure.ExtendDefaultLuminanceRange=true
 r.AllowStaticLighting=False
 
@@ -40,5 +83,9 @@ SoundCueCookQualityIndex=-1
 -TargetedRHIs=SF_VULKAN_SM5
 +TargetedRHIs=SF_VULKAN_SM6
 
+[/Script/MacTargetPlatform.MacTargetSettings]
+-TargetedRHIs=SF_METAL_SM5
++TargetedRHIs=SF_METAL_SM6
+
 [/Script/NavigationSystem.RecastNavMesh]
 RuntimeGeneration=Dynamic
\ No newline at end of file
diff --git a/Templates/TP_TopDownBP/Config/DefaultEngine.ini b/Templates/TP_TopDownBP/Config/DefaultEngine.ini
--- a/Templates/TP_TopDownBP/Config/DefaultEngine.ini
+++ b/Templates/TP_TopDownBP/Config/DefaultEngine.ini
@@ -4,10 +4,6 @@ GameName=TP_TopDownBP
 [/Script/EngineSettings.GameMapsSettings]
 EditorStartupMap=/Game/TopDown/Lvl_TopDown.Lvl_TopDown
 GameDefaultMap=/Game/TopDown/Lvl_TopDown.Lvl_TopDown
-TransitionMap=
-bUseSplitscreen=True
-TwoPlayerSplitscreenLayout=Horizontal
-ThreePlayerSplitscreenLayout=FavorTop
 GlobalDefaultGameMode=/Game/TopDown/Blueprints/BP_TopDownGameMode.BP_TopDownGameMode_C
 GlobalDefaultServerGameMode=None
 
@@ -90,3 +86,10 @@ SoundCueCookQualityIndex=-1
 [/Script/LinuxTargetPlatform.LinuxTargetSettings]
 -TargetedRHIs=SF_VULKAN_SM5
 +TargetedRHIs=SF_VULKAN_SM6
+
+[/Script/MacTargetPlatform.MacTargetSettings]
+-TargetedRHIs=SF_METAL_SM5
++TargetedRHIs=SF_METAL_SM6
+
+[/Script/NavigationSystem.RecastNavMesh]
+RuntimeGeneration=Dynamic
diff --git a/Templates/TP_VehicleAdv/Config/DefaultEngine.ini b/Templates/TP_VehicleAdv/Config/DefaultEngine.ini
--- a/Templates/TP_VehicleAdv/Config/DefaultEngine.ini
+++ b/Templates/TP_VehicleAdv/Config/DefaultEngine.ini
@@ -1,6 +1,6 @@
 [/Script/Engine.PhysicsSettings]
 bSubstepping=False
-bSubsteppingAsync=False
+bSubsteppingAsync=True
 
 [/Script/EngineSettings.GameMapsSettings]
 EditorStartupMap=/Game/VehicleTemplate/Maps/VehicleBasic.VehicleBasic
@@ -19,8 +19,23 @@ GlobalDefaultServerGameMode=None
 r.ReflectionMethod=1
 r.GenerateMeshDistanceFields=True
 r.DynamicGlobalIlluminationMethod=1
+r.ClearCoatNormal=True
+r.DefaultFeature.AutoExposure.ExtendDefaultLuminanceRange=True
+r.DefaultFeature.AutoExposure.ExtendDefaultLuminanceRange=true
 r.AllowStaticLighting=False
 
+r.SkinCache.CompileShaders=True
+
+r.RayTracing=True
+
+r.Shadow.Virtual.Enable=1
+
+r.DefaultFeature.LocalExposure.HighlightContrastScale=0.8
+
+r.DefaultFeature.LocalExposure.ShadowContrastScale=0.8
+r.DefaultFeature.LightUnits=2
+r.VirtualTextures=True
+
 [/Script/WindowsTargetPlatform.WindowsTargetSettings]
 DefaultGraphicsRHI=DefaultGraphicsRHI_DX12
 -D3D12TargetedShaderFormats=PCD3D_SM5
@@ -53,3 +68,7 @@ SoundCueCookQualityIndex=-1
 [/Script/LinuxTargetPlatform.LinuxTargetSettings]
 -TargetedRHIs=SF_VULKAN_SM5
 +TargetedRHIs=SF_VULKAN_SM6
+
+[/Script/MacTargetPlatform.MacTargetSettings]
+-TargetedRHIs=SF_METAL_SM5
++TargetedRHIs=SF_METAL_SM6
diff --git a/Templates/TP_VehicleAdvBP/Config/DefaultEngine.ini b/Templates/TP_VehicleAdvBP/Config/DefaultEngine.ini
--- a/Templates/TP_VehicleAdvBP/Config/DefaultEngine.ini
+++ b/Templates/TP_VehicleAdvBP/Config/DefaultEngine.ini
@@ -41,7 +41,6 @@ r.VirtualTextures=True
 
 [/Script/WindowsTargetPlatform.WindowsTargetSettings]
 DefaultGraphicsRHI=DefaultGraphicsRHI_DX12
-DefaultGraphicsRHI=DefaultGraphicsRHI_DX12
 -D3D12TargetedShaderFormats=PCD3D_SM5
 +D3D12TargetedShaderFormats=PCD3D_SM6
 -D3D11TargetedShaderFormats=PCD3D_SM5
@@ -73,8 +72,6 @@ SoundCueCookQualityIndex=-1
 -TargetedRHIs=SF_VULKAN_SM5
 +TargetedRHIs=SF_VULKAN_SM6
 
-[/Script/HardwareTargeting.HardwareTargetingSettings]
-TargetedHardwareClass=Desktop
-AppliedTargetedHardwareClass=Desktop
-DefaultGraphicsPerformance=Maximum
-AppliedDefaultGraphicsPerformance=Maximum
\ No newline at end of file
+[/Script/MacTargetPlatform.MacTargetSettings]
+-TargetedRHIs=SF_METAL_SM5
++TargetedRHIs=SF_METAL_SM6
\ No newline at end of file
